Create public and restricted versions of LowLevelTests_GenProps.xml

diff --git a/Engine/Build/LowLevelTests_GenProps.xml b/Engine/Build/LowLevelTests_GenProps.xml
new file mode 100644
--- /dev/null
+++ b/Engine/Build/LowLevelTests_GenProps.xml
@@ -0,0 +1,11 @@
+﻿<?xml version="1.0" encoding="utf-8"?>
+<BuildGraph xmlns:xsi="http://www.w3.org/2001/XMLSchema-instance" xsi:schemaLocation="http://www.epicgames.com/BuildGraph ../../Build/Graph/Schema.xsd" xmlns="http://www.epicgames.com/BuildGraph">
+  <Property Name="TestNames" Value="HeadlessChaos" />
+  <Property Name="HeadlessChaosShort" Value="Headless Chaos" />
+  <Property Name="HeadlessChaosTarget" Value="HeadlessChaos" />
+  <Property Name="HeadlessChaosBinariesRelative" Value="Engine\Binaries" />
+  <Property Name="HeadlessChaosResources" Value="" />
+  <Property Name="HeadlessChaosResourcesDest" Value="" />
+  <Option Name="RunHeadlessChaosTests" DefaultValue="False" Description="Run Headless Chaos Tests" />
+  <Property Name="HeadlessChaosSupportedPlatforms" Value="Win64" />
+</BuildGraph>
\ No newline at end of file
diff --git a/Engine/Source/Programs/HeadlessChaos/HeadlessChaos.Build.cs b/Engine/Source/Programs/HeadlessChaos/HeadlessChaos.Build.cs
--- a/Engine/Source/Programs/HeadlessChaos/HeadlessChaos.Build.cs
+++ b/Engine/Source/Programs/HeadlessChaos/HeadlessChaos.Build.cs
@@ -4,8 +4,13 @@ using System.IO;
 using UnrealBuildTool;
 
 [SupportedPlatforms("Win64")]
-public class HeadlessChaos : ModuleRules
+public class HeadlessChaos : LowLevelTests
 {
+	public override string TestName => "HeadlessChaos";
+	public override string TestShortName => "Headless Chaos";
+
+	public override bool UsesCatch2 => false;
+
 	public HeadlessChaos(ReadOnlyTargetRules Target) : base(Target)
 	{
 		PublicIncludePaths.Add("Runtime/Launch/Public");
@@ -50,12 +55,5 @@ public class HeadlessChaos : ModuleRules
 		}
 
 		PrivateDefinitions.Add("CHAOS_INCLUDE_LEVEL_1=1");
-		
-		LowLevelTests.UpdateGeneratedPropertiesScriptFile(
-			typeof(HeadlessChaos),
-			"HeadlessChaos",
-			"Headless Chaos",
-			Target.LaunchModuleName,
-			Path.Combine("Engine", "Binaries"));
 	}
 }
diff --git a/Engine/Source/Programs/LowLevelTests/LowLevelTests.Build.cs b/Engine/Source/Programs/LowLevelTests/LowLevelTests.Build.cs
--- a/Engine/Source/Programs/LowLevelTests/LowLevelTests.Build.cs
+++ b/Engine/Source/Programs/LowLevelTests/LowLevelTests.Build.cs
@@ -9,64 +9,92 @@ using UnrealBuildTool;
 
 public class LowLevelTests : ModuleRules
 {
-	private static string GeneratedPropertiesScriptFile = null;
+	private readonly XNamespace BuildGraphNamespace = XNamespace.Get("http://www.epicgames.com/BuildGraph");
+	private readonly XNamespace SchemaInstance = XNamespace.Get("http://www.w3.org/2001/XMLSchema-instance");
+	private readonly XNamespace SchemaLocation = XNamespace.Get("http://www.epicgames.com/BuildGraph ../../Build/Graph/Schema.xsd");
 
-	private static XNamespace xmlns = XNamespace.Get("http://www.epicgames.com/BuildGraph");
-	private static XNamespace xsi = XNamespace.Get("http://www.w3.org/2001/XMLSchema-instance");
-	private static XNamespace schemaLocation = XNamespace.Get("http://www.epicgames.com/BuildGraph ../../Build/Graph/Schema.xsd");
+	public virtual string TestName => throw new Exception("TestName must be overwritten in subclasses.");
+	public virtual string TestShortName => throw new Exception("TestShortName must be overwritten in subclasses.");
+	public virtual string ResourcesPath => string.Empty;
+
+	public virtual bool UsesCatch2 => true;
 
 	public LowLevelTests(ReadOnlyTargetRules Target) : base(Target)
 	{
-		PCHUsage = PCHUsageMode.NoPCHs;
-		PrecompileForTargets = PrecompileTargetsType.None;
-
-		if (Target.Configuration == UnrealTargetConfiguration.Debug && Target.Platform == UnrealTargetPlatform.Linux)
+		if (UsesCatch2)
 		{
-			OptimizeCode = CodeOptimization.Never;
-		}
+			PCHUsage = PCHUsageMode.NoPCHs;
+			PrecompileForTargets = PrecompileTargetsType.None;
 
-		bAllowConfidentialPlatformDefines = true;
-		bLegalToDistributeObjectCode = true;
+			if (Target.Configuration == UnrealTargetConfiguration.Debug && Target.Platform == UnrealTargetPlatform.Linux)
+			{
+				OptimizeCode = CodeOptimization.Never;
+			}
+
+			bAllowConfidentialPlatformDefines = true;
+			bLegalToDistributeObjectCode = true;
 
-		// Required false for catch.hpp
-		bUseUnity = false;
+			// Required false for catch.hpp
+			bUseUnity = false;
 
-		// Disable exception handling so that tests can assert for exceptions
-		bEnableObjCExceptions = false;
-		bEnableExceptions = false;
+			// Disable exception handling so that tests can assert for exceptions
+			bEnableObjCExceptions = false;
+			bEnableExceptions = false;
 
-		PrivateDependencyModuleNames.AddRange(
-			new string[] {
+			PrivateDependencyModuleNames.AddRange(
+				new string[] {
 				"Core",
 				"Projects",
 				"LowLevelTestsRunner"
-			});
+				});
 
-		PrivateIncludePaths.Add("Runtime/Launch/Private");
-		PrivateIncludePathModuleNames.Add("Launch");
+			PrivateIncludePaths.Add("Runtime/Launch/Private");
+			PrivateIncludePathModuleNames.Add("Launch");
 
-		// Platforms specific setup
-		if (Target.Platform == UnrealTargetPlatform.Android)
+			// Platforms specific setup
+			if (Target.Platform == UnrealTargetPlatform.Android)
+			{
+				PublicDefinitions.Add("CATCH_CONFIG_NOSTDOUT");
+			}
+		}
+
+		if (this.GetType() != typeof(LowLevelTests))
 		{
-			PublicDefinitions.Add("CATCH_CONFIG_NOSTDOUT");
+			UpdateBuildGraphPropertiesFile();
 		}
 	}
 
 	/// <summary>
 	/// Generates or updates include file for LowLevelTests.xml containing test flags: name, short name, target name, relative binaries path, supported platforms etc.
-	/// Static for now until HeadlessChaos can inherit from LowLevelTests.	
+	/// <paramref name="TestModule">The test module build class that inherits form LowLevelTests</paramref>
 	/// </summary>
-	/// <param name="ModuleClassType">Type of test build module class that inherits from LowLevelTests.</param>
-	/// <param name="TestName">Name of low level test.</param>
-	/// <param name="TestShortName">Test short name.</param>
-	/// <param name="TestTargetName">Test target executable name.</param>
-	/// <param name="BinariesPath">Test binaries path.</param>
-	/// <param name="ResourcesPath">Optional test resources to be copied.</param>
-	public static void UpdateGeneratedPropertiesScriptFile(Type ModuleClassType, string TestName, string TestShortName, string TestTargetName, string BinariesPath, string ResourcesPath = "")
+	private void UpdateBuildGraphPropertiesFile()
 	{
-		if (GeneratedPropertiesScriptFile == null)
+		bool IsPublic = false;
+		string GeneratedPropertiesScriptFile;
+
+		string RestrictedFolder = Path.Combine(Unreal.EngineDirectory.FullName, "Restricted");
+		string NotForLicenseesFolder = Path.Combine(RestrictedFolder, "NotForLicensees");
+		string NonPublicFolder = Path.Combine(NotForLicenseesFolder, "Build");
+		string NonPublicPath = Path.Combine(NonPublicFolder, "LowLevelTests_GenProps.xml");
+
+		if (IsRestrictedPath(ModuleDirectory))
 		{
-			GeneratedPropertiesScriptFile = Path.Combine(Unreal.EngineDirectory.FullName, "Restricted", "NotForLicensees", "Build", "LowLevelTests_GenProps.xml");
+			GeneratedPropertiesScriptFile = NonPublicPath;
+		}
+		else
+		{
+			IsPublic = true;
+			GeneratedPropertiesScriptFile = Path.Combine(Unreal.EngineDirectory.FullName, "Build", "LowLevelTests_GenProps.xml");
+		}
+
+		// UE-133126
+		if (Directory.GetFileSystemEntries(NonPublicFolder).Length == 1 && File.Exists(NonPublicPath))
+		{
+			File.Delete(NonPublicPath);
+			Directory.Delete(NonPublicFolder);
+			Directory.Delete(NotForLicenseesFolder);
+			Directory.Delete(RestrictedFolder);
 		}
 
 		if (!File.Exists(GeneratedPropertiesScriptFile))
@@ -74,69 +102,79 @@ public class LowLevelTests : ModuleRules
 			Directory.CreateDirectory(Path.GetDirectoryName(GeneratedPropertiesScriptFile));
 			using (FileStream fileStream = File.Create(GeneratedPropertiesScriptFile))
 			{
-				XDocument initFile = new XDocument(new XElement(xmlns + "BuildGraph", new XAttribute(XNamespace.Xmlns + "xsi", xsi), new XAttribute(xsi + "schemaLocation", schemaLocation)));
-				initFile.Root.Add(new XElement(xmlns + "Property", new XAttribute("Name", "TestNames"), new XAttribute("Value", string.Empty)));
+				XDocument initFile = new XDocument(new XElement(BuildGraphNamespace + "BuildGraph", new XAttribute(XNamespace.Xmlns + "xsi", SchemaInstance), new XAttribute(SchemaInstance + "schemaLocation", SchemaLocation)));
+				initFile.Root.Add(
+					new XElement(
+						BuildGraphNamespace + "Property",
+						new XAttribute("Name", "TestNames" + (!IsPublic ? "Restricted" : "")),
+						new XAttribute("Value", string.Empty)));
 
 				initFile.Save(fileStream);
 			}
 		}
-		XDocument genPropsDoc = XDocument.Load(GeneratedPropertiesScriptFile);
-		XElement root = genPropsDoc.Root;
+
+		// All relevant properties
+		string TestTargetName = Target.LaunchModuleName;
+		string TestBinariesPath = TryGetBinariesPath();
+		string TestResourcesPath = ResourcesPath;
+
+		// Do not save full paths
+		if (Path.IsPathRooted(TestBinariesPath))
+		{
+			TestBinariesPath = Path.GetRelativePath(Unreal.RootDirectory.FullName, TestBinariesPath);
+		}
+		if (Path.IsPathRooted(ResourcesPath))
+		{
+			TestResourcesPath = Path.GetRelativePath(Unreal.RootDirectory.FullName, ResourcesPath);
+		}
+
+		XDocument GenPropsDoc = XDocument.Load(GeneratedPropertiesScriptFile);
+		XElement Root = GenPropsDoc.Root;
 		// First descendant must be TestNames
-		XElement testNames = (XElement)root.FirstNode;
-		string allTestNames = testNames.Attribute("Value").Value;
-		if (!allTestNames.Contains(TestName))
+		XElement TestNames = (XElement)Root.FirstNode;
+		string AllTestNames = TestNames.Attribute("Value").Value;
+		if (!AllTestNames.Contains(TestName))
 		{
-			if (string.IsNullOrEmpty(allTestNames))
+			if (string.IsNullOrEmpty(AllTestNames))
 			{
-				allTestNames += TestName;
+				AllTestNames += TestName;
 			}
-			else if (!allTestNames.Contains(TestName))
+			else if (!AllTestNames.Contains(TestName))
 			{
-				allTestNames += ";" + TestName;
+				AllTestNames += ";" + TestName;
 			}
 		}
-		testNames.Attribute("Value").SetValue(allTestNames);
+		TestNames.Attribute("Value").SetValue(AllTestNames);
 
-		XElement lastUpdatedNode = testNames;
+		XElement lastUpdatedNode = TestNames;
 		InsertOrUpdateTestFlag(ref lastUpdatedNode, TestName, "Short", TestShortName);
 		InsertOrUpdateTestFlag(ref lastUpdatedNode, TestName, "Target", TestTargetName);
-
-		if (Path.IsPathRooted(BinariesPath))
-		{
-			BinariesPath = Path.GetRelativePath(Unreal.RootDirectory.FullName, BinariesPath);
-		}
-		InsertOrUpdateTestFlag(ref lastUpdatedNode, TestName, "BinariesRelative", BinariesPath);
-
-		if (Path.IsPathRooted(ResourcesPath))
-		{
-			ResourcesPath = Path.GetRelativePath(Unreal.RootDirectory.FullName, ResourcesPath);
-		}
-		InsertOrUpdateTestFlag(ref lastUpdatedNode, TestName, "Resources", ResourcesPath);
-		InsertOrUpdateTestFlag(ref lastUpdatedNode, TestName, "ResourcesDest", Path.GetFileName(ResourcesPath));
+		InsertOrUpdateTestFlag(ref lastUpdatedNode, TestName, "BinariesRelative", TestBinariesPath);
+		InsertOrUpdateTestFlag(ref lastUpdatedNode, TestName, "Resources", TestResourcesPath);
+		InsertOrUpdateTestFlag(ref lastUpdatedNode, TestName, "ResourcesDest", Path.GetFileName(TestResourcesPath));
 
 		InsertOrUpdateTestOption(ref lastUpdatedNode, TestName, TestShortName, "Run", "Tests", false.ToString());
 
-		List<UnrealTargetPlatform> allSupportedPlatforms = new List<UnrealTargetPlatform>();
-		var supportedPlatforms = ModuleClassType.GetCustomAttributes(typeof(SupportedPlatformsAttribute), false);
+		List<UnrealTargetPlatform> AllSupportedPlatforms = new List<UnrealTargetPlatform>();
+		var SupportedPlatforms = GetType().GetCustomAttributes(typeof(SupportedPlatformsAttribute), false);
 		// If none specified we assume all platforms are supported by default
-		if (supportedPlatforms.Length == 0)
+		if (SupportedPlatforms.Length == 0)
 		{
-			allSupportedPlatforms.AddRange(UnrealTargetPlatform.GetValidPlatforms().ToList());
+			AllSupportedPlatforms.AddRange(UnrealTargetPlatform.GetValidPlatforms().ToList());
 		}
 		else
 		{
-			foreach (var platform in supportedPlatforms)
+			foreach (var Platform in SupportedPlatforms)
 			{
-				allSupportedPlatforms.AddRange(((SupportedPlatformsAttribute)platform).Platforms);
+				AllSupportedPlatforms.AddRange(((SupportedPlatformsAttribute)Platform).Platforms);
 			}
 		}
 
-		InsertOrUpdateTestFlag(ref lastUpdatedNode, TestName, "SupportedPlatforms", allSupportedPlatforms.Aggregate("", (current, next) => (current == "" ? next.ToString() : current + ";" + next.ToString())));
+		InsertOrUpdateTestFlag(ref lastUpdatedNode, TestName, "SupportedPlatforms", AllSupportedPlatforms.Aggregate("", (current, next) => (current == "" ? next.ToString() : current + ";" + next.ToString())));
 
 		try
 		{
-			genPropsDoc.Save(GeneratedPropertiesScriptFile);
+			GenPropsDoc.Save(GeneratedPropertiesScriptFile);
 		}
 		catch (UnauthorizedAccessException)
 		{
@@ -145,19 +183,43 @@ public class LowLevelTests : ModuleRules
 		}
 	}
 
-	private static void InsertOrUpdateTestFlag(ref XElement ElementUpsertAfter, string TestName, string FlagSuffix, string FlagValue)
+	private bool IsRestrictedPath(string ModuleDirectory)
 	{
-		IEnumerable<XElement> NextChunk = ElementUpsertAfter.ElementsAfterSelf(xmlns + "Property")
+		foreach(string RestrictedFolderName in RestrictedFolder.GetNames())
+		{
+			if (ModuleDirectory.Contains(RestrictedFolderName))
+			{
+				return true;
+			}
+		}
+
+		return false;
+	}
+
+	private string TryGetBinariesPath()
+	{
+		int SourceFolderIndex = ModuleDirectory.IndexOf("Source");
+		if (SourceFolderIndex < 0)
+		{
+			throw new Exception("Could not detect source folder path for module " + GetType());
+		}
+		return ModuleDirectory.Substring(0, SourceFolderIndex) + "Binaries";
+	}
+
+	private void InsertOrUpdateTestFlag(ref XElement ElementUpsertAfter, string TestName, string FlagSuffix, string FlagValue)
+	{
+		IEnumerable<XElement> NextChunk = ElementUpsertAfter.ElementsAfterSelf(BuildGraphNamespace + "Property")
 			.Where(prop => prop.Attribute("Name").Value.EndsWith(FlagSuffix));
 		if (NextChunk
 			.Where(prop => prop.Attribute("Name").Value == TestName + FlagSuffix)
 			.Count() == 0)
 		{
-			XElement ElementInsert = new XElement(xmlns + "Property");
+			XElement ElementInsert = new XElement(BuildGraphNamespace + "Property");
 			ElementInsert.SetAttributeValue("Name", TestName + FlagSuffix);
 			ElementInsert.SetAttributeValue("Value", FlagValue);
 			ElementUpsertAfter.AddAfterSelf(ElementInsert);
-		} else
+		}
+		else
 		{
 			NextChunk
 				.Where(prop => prop.Attribute("Name").Value == TestName + FlagSuffix).First().SetAttributeValue("Value", FlagValue);
@@ -165,20 +227,21 @@ public class LowLevelTests : ModuleRules
 		ElementUpsertAfter = NextChunk.Last();
 	}
 
-	private static void InsertOrUpdateTestOption(ref XElement ElementUpsertAfter, string TestName, string TestShortName, string OptionPrefix, string OptionSuffix, string DefaultValue)
+	private void InsertOrUpdateTestOption(ref XElement ElementUpsertAfter, string TestName, string TestShortName, string OptionPrefix, string OptionSuffix, string DefaultValue)
 	{
-		IEnumerable<XElement> NextChunk = ElementUpsertAfter.ElementsAfterSelf(xmlns + "Option")
+		IEnumerable<XElement> NextChunk = ElementUpsertAfter.ElementsAfterSelf(BuildGraphNamespace + "Option")
 			.Where(prop => prop.Attribute("Name").Value.StartsWith(OptionPrefix) && prop.Attribute("Name").Value.EndsWith(OptionSuffix));
 		if (NextChunk
 			.Where(prop => prop.Attribute("Name").Value == OptionPrefix + TestName + OptionSuffix)
 			.Count() == 0)
 		{
-			XElement ElementInsert = new XElement(xmlns + "Option");
+			XElement ElementInsert = new XElement(BuildGraphNamespace + "Option");
 			ElementInsert.SetAttributeValue("Name", OptionPrefix + TestName + OptionSuffix);
 			ElementInsert.SetAttributeValue("DefaultValue", DefaultValue);
 			ElementInsert.SetAttributeValue("Description", string.Format("{0} {1} {2}", OptionPrefix, TestShortName, OptionSuffix));
 			ElementUpsertAfter.AddAfterSelf(ElementInsert);
-		} else
+		}
+		else
 		{
 			XElement ElementUpdate = NextChunk
 				.Where(prop => prop.Attribute("Name").Value == OptionPrefix + TestName + OptionSuffix).First();
diff --git a/Engine/Source/Programs/LowLevelTests/README.md b/Engine/Source/Programs/LowLevelTests/README.md
--- a/Engine/Source/Programs/LowLevelTests/README.md
+++ b/Engine/Source/Programs/LowLevelTests/README.md
@@ -16,9 +16,7 @@ Catch2 is a modern C++ test framework designed with multi-platform support and i
 ## <a name="CreateTestModule"/>Create a test module
 Create a module that inherits from `LowLevelTests`.
 
-Include a call to `UpdateGeneratedPropertiesScriptFile` - this method generates test flag metadata required when running tests with BuildGraph.
-
-Update *Engine\Restricted\NotForLicensees\Build\LowLevelTests.xml* (details below)
+Override TestName and TestShortName with appropriate values.
 
 **MyModuleTests.Build.cs**
 ```csharp
@@ -26,6 +24,9 @@ using UnrealBuildTool;
 
 public class MyModuleTests : LowLevelTests
 {
+	public override string TestName => "MyModule";
+	public override string TestShortName => "My Module";
+
 	public MyModuleTests(ReadOnlyTargetRules Target) : base(Target)
 	{
 		PrivateIncludePaths.AddRange(
@@ -37,12 +38,6 @@ public class MyModuleTests : LowLevelTests
 		);
 		
 		// Other types of dependencies minimally required to test a specific module
-
-		UpdateGeneratedPropertiesScriptFile(typeof(MyModuleTests),
-			"MyModuleTestName",
-			"My Module Short Name",
-			Target.LaunchModuleName,
-			Path.Combine("Path", "To", "Binaries"));
 	}
 }
 
@@ -84,7 +79,7 @@ You can now build the generated project and the executable will be placed in the
 
 **Buildgraph**
 
-The buildgraph file *Engine\Restricted\NotForLicensees\Build\LowLevelTests.xml* needs to be updated with an option to run the new test module we just added.
+The buildgraph file *Engine\Restricted\NotForLicensees\Build\LowLevelTests.xml* has an option to run the new test module we just added.
 
 This option can be used to selectively include or exclude the test when running via BuildGraph.
 
