LLTs: add disabled flag, enable it for OSS
#rnx

diff --git a/Engine/Build/LowLevelTests_GenProps.xml b/Engine/Build/LowLevelTests_GenProps.xml
--- a/Engine/Build/LowLevelTests_GenProps.xml
+++ b/Engine/Build/LowLevelTests_GenProps.xml
@@ -1,6 +1,7 @@
 ﻿<?xml version="1.0" encoding="utf-8"?>
 <BuildGraph xmlns:xsi="http://www.w3.org/2001/XMLSchema-instance" xsi:schemaLocation="http://www.epicgames.com/BuildGraph ../../Build/Graph/Schema.xsd" xmlns="http://www.epicgames.com/BuildGraph">
   <Property Name="TestNames" Value="HeadlessChaos" />
+  <Property Name="HeadlessChaosDisabled" Value="False" />
   <Property Name="HeadlessChaosShort" Value="Headless Chaos" />
   <Property Name="HeadlessChaosTarget" Value="HeadlessChaos" />
   <Property Name="HeadlessChaosBinariesRelative" Value="Engine\Binaries" />
diff --git a/Engine/Source/Programs/LowLevelTests/LowLevelTests.Build.cs b/Engine/Source/Programs/LowLevelTests/LowLevelTests.Build.cs
--- a/Engine/Source/Programs/LowLevelTests/LowLevelTests.Build.cs
+++ b/Engine/Source/Programs/LowLevelTests/LowLevelTests.Build.cs
@@ -19,6 +19,8 @@ public class LowLevelTests : ModuleRules
 
 	public virtual bool UsesCatch2 => true;
 
+	public virtual bool Disabled => false;
+
 	public LowLevelTests(ReadOnlyTargetRules Target) : base(Target)
 	{
 		if (UsesCatch2)
@@ -147,6 +149,7 @@ public class LowLevelTests : ModuleRules
 		TestNames.Attribute("Value").SetValue(AllTestNames);
 
 		XElement lastUpdatedNode = TestNames;
+		InsertOrUpdateTestFlag(ref lastUpdatedNode, TestName, "Disabled", Disabled.ToString());
 		InsertOrUpdateTestFlag(ref lastUpdatedNode, TestName, "Short", TestShortName);
 		InsertOrUpdateTestFlag(ref lastUpdatedNode, TestName, "Target", TestTargetName);
 		InsertOrUpdateTestFlag(ref lastUpdatedNode, TestName, "BinariesRelative", TestBinariesPath);
