Fix public facing error "Unable to instantiate module 'HeadlessChaos': System.IO.DirectoryNotFoundException: Could not find a part of the path"

diff --git a/Engine/Source/Programs/LowLevelTests/LowLevelTests.Build.cs b/Engine/Source/Programs/LowLevelTests/LowLevelTests.Build.cs
--- a/Engine/Source/Programs/LowLevelTests/LowLevelTests.Build.cs
+++ b/Engine/Source/Programs/LowLevelTests/LowLevelTests.Build.cs
@@ -91,7 +91,7 @@ public class LowLevelTests : ModuleRules
 		}
 
 		// UE-133126
-		if (Directory.GetFileSystemEntries(NonPublicFolder).Length == 1 && File.Exists(NonPublicPath))
+		if (File.Exists(NonPublicPath) && Directory.GetFileSystemEntries(NonPublicFolder).Length == 1)
 		{
 			File.Delete(NonPublicPath);
 			Directory.Delete(NonPublicFolder);
