IOS Self test - use app core.

#rnx

diff --git a/Engine/Source/Programs/LowLevelTests/LowLevelTests.Target.cs b/Engine/Source/Programs/LowLevelTests/LowLevelTests.Target.cs
--- a/Engine/Source/Programs/LowLevelTests/LowLevelTests.Target.cs
+++ b/Engine/Source/Programs/LowLevelTests/LowLevelTests.Target.cs
@@ -66,6 +66,10 @@ public class LowLevelTestsTarget : TargetRules
 		else if (Target.Platform == UnrealTargetPlatform.IOS)
 		{
 			GlobalDefinitions.Add("HAS_METAL=0");
+
+			bIsBuildingConsoleApplication = false;
+			// Required for IOS, but needs to fix compilation errors
+			bCompileAgainstApplicationCore = true;
 		}
 	}
 
