LowLevelTests: Add ReportType to set -r xml for Dashboard reporting output

diff --git a/Engine/Build/LowLevelTests_GenProps.xml b/Engine/Build/LowLevelTests_GenProps.xml
--- a/Engine/Build/LowLevelTests_GenProps.xml
+++ b/Engine/Build/LowLevelTests_GenProps.xml
@@ -7,6 +7,7 @@
   <Property Name="HeadlessChaosBinariesRelative" Value="Engine\Binaries" />
   <Property Name="HeadlessChaosResources" Value="" />
   <Property Name="HeadlessChaosResourcesDest" Value="" />
+  <Property Name="HeadlessChaosReportType" Value="console" />
   <Option Name="RunHeadlessChaosTests" DefaultValue="False" Description="Run Headless Chaos Tests" />
   <Property Name="HeadlessChaosSupportedPlatforms" Value="Win64" />
 </BuildGraph>
\ No newline at end of file
diff --git a/Engine/Source/Programs/LowLevelTests/LowLevelTests.Build.cs b/Engine/Source/Programs/LowLevelTests/LowLevelTests.Build.cs
--- a/Engine/Source/Programs/LowLevelTests/LowLevelTests.Build.cs
+++ b/Engine/Source/Programs/LowLevelTests/LowLevelTests.Build.cs
@@ -15,6 +15,8 @@ public class LowLevelTests : ModuleRules
 
 	public virtual string TestName => throw new Exception("TestName must be overwritten in subclasses.");
 	public virtual string TestShortName => throw new Exception("TestShortName must be overwritten in subclasses.");
+	
+	public virtual string ReportType => "console";
 	public virtual string ResourcesPath => string.Empty;
 
 	public virtual bool UsesCatch2 => true;
@@ -173,6 +175,7 @@ public class LowLevelTests : ModuleRules
 		InsertOrUpdateTestFlag(ref lastUpdatedNode, TestName, "BinariesRelative", TestBinariesPath);
 		InsertOrUpdateTestFlag(ref lastUpdatedNode, TestName, "Resources", TestResourcesPath);
 		InsertOrUpdateTestFlag(ref lastUpdatedNode, TestName, "ResourcesDest", Path.GetFileName(TestResourcesPath));
+		InsertOrUpdateTestFlag(ref lastUpdatedNode, TestName, "ReportType", ReportType.ToString());
 
 		InsertOrUpdateTestOption(ref lastUpdatedNode, TestName, TestShortName, "Run", "Tests", false.ToString());
 
