Copyright boilerplate
#rnx

diff --git a/Engine/Source/Programs/LowLevelTests/Tests/LowLevelTestsModule.cpp b/Engine/Source/Programs/LowLevelTests/Tests/LowLevelTestsModule.cpp
--- a/Engine/Source/Programs/LowLevelTests/Tests/LowLevelTestsModule.cpp
+++ b/Engine/Source/Programs/LowLevelTests/Tests/LowLevelTestsModule.cpp
@@ -1,3 +1,4 @@
+// Copyright Epic Games, Inc. All Rights Reserved.
 #include "TestCommon/CoreUtilities.h"
 #include "TestCommon/CoreUObjectUtilities.h"
 
@@ -29,4 +30,4 @@ void FGlobalLLTModule::GlobalTeardown()
 #endif
 	CleanupThreadPool();
 	CleanupTaskGraph();
-}
\ No newline at end of file
+}
