Low level tests readme update
#rnx

diff --git a/Engine/Source/Programs/LowLevelTests/README.md b/Engine/Source/Programs/LowLevelTests/README.md
--- a/Engine/Source/Programs/LowLevelTests/README.md
+++ b/Engine/Source/Programs/LowLevelTests/README.md
@@ -14,19 +14,17 @@ Developing tests in the classic pyramid style of unit, integration and functiona
 Catch2 is a modern C++ test framework designed with multi-platform support and it can run individual tests or a filtered subset of them using tags.
 
 ## <a name="CreateTestModule"/>Create a test module
-Create a module that inherits from `LowLevelTests`.
+Create a module that inherits from `TestModuleRules`.
+If building a test module that doesn't use Catch2, call the base constructor `: base(Target, false)`.
 
-Override TestName and TestShortName with appropriate values.
+Call UpdateBuildGraphPropertiesFile if intended to generate BuildGraph property metadata.
 
 **MyModuleTests.Build.cs**
 ```csharp
 using UnrealBuildTool;
 
-public class MyModuleTests : LowLevelTests
+public class MyModuleTests : TestModuleRules
 {
-	public override string TestName => "MyModule";
-	public override string TestShortName => "My Module";
-
 	public MyModuleTests(ReadOnlyTargetRules Target) : base(Target)
 	{
 		PrivateIncludePaths.AddRange(
@@ -38,13 +36,15 @@ public class MyModuleTests : LowLevelTests
 		);
 		
 		// Other types of dependencies minimally required to test a specific module
+
+		UpdateBuildGraphPropertiesFile(new Metadata("MyModule", "My Module"));
 	}
 }
 
 ```
 ---
-Next add a target file and override default compilation flags if necessary - aim for a minimal testable module free of default features that don't add value to low level testing.
-Specify the supported platforms and make sure to call `SetupModule()` as shown below.
+Next add a target file with a class that inherits from `TestTargetRules` and override default compilation flags if necessary - aim for a minimal testable module free of default features that don't add value to low level testing.
+You can specify the supported platforms individually, if not they default to Win64, Mac, Linux and Android.
 
 Here is where you can enable project specific global definitions as well as `Catch2` macros such as benchmarking or compilation flags etc.
 
@@ -53,12 +53,10 @@ Here is where you can enable project specific global definitions as well as `Cat
 using UnrealBuildTool;
 
 [SupportedPlatforms(UnrealPlatformClass.All)]
-public class MyModuleTestsTarget : LowLevelTestsTarget
+public class MyModuleTestsTarget : TestTargetRules
 {
 	public MyModuleTestsTarget(TargetInfo Target) : base(Target)
 	{
-		SetupModule();
-
 		bCompileAgainstCoreUObject = true;
 		bCompileAgainstApplicationCore = true;
 
@@ -69,7 +67,7 @@ public class MyModuleTestsTarget : LowLevelTestsTarget
 ```
 ---
 
-**Make sure to checkout Engine\Restricted\NotForLicensees\Build\LowLevelTests_GenProps.xml before moving on - this is where the BuildGraph test related metadata will be generated.**
+**Make sure to checkout Engine\Build\LowLevelTests_GenProps.xml before moving on - this is where the BuildGraph test related metadata will be generated.**
 
 Generate the solution w/ GenerateProjectFiles.bat - the test project will be placed in the Visual Studio solution folder `Programs/LowLevelTests`.
 
@@ -79,7 +77,7 @@ You can now build the generated project and the executable will be placed in the
 
 **Buildgraph**
 
-The buildgraph file *Engine\Restricted\NotForLicensees\Build\LowLevelTests.xml* has an option to run the new test module we just added.
+The buildgraph file *Engine\Build\LowLevelTests.xml* has an option to run the new test module we just added.
 
 This option can be used to selectively include or exclude the test when running via BuildGraph.
 
@@ -88,12 +86,10 @@ This option can be used to selectively include or exclude the test when running
 ```
 
 ```shell 
-.\RunUAT.bat BuildGraph -Script="Engine/Restricted/NotForLicensees/Build/LowLevelTests.xml" -set:RunMyModuleTests=true
+.\RunUAT.bat BuildGraph -Script="Engine/Build/LowLevelTests.xml" -set:RunMyModuleTests=true
 ```
 
 
-
-
 ## <a name="AddTests"/>Add tests
 
 ### Overview of types of tests
@@ -120,30 +116,6 @@ This option can be used to selectively include or exclude the test when running
 - Functional tests can take anywhere from seconds to minutes and rarely hours to run.
 They are usually slower than integration tests.
 
-### main.cpp
-
-Add a minimal `main.cpp` file into your tests module *Private* folder:
-```cpp
-#include "TestRunner.h"
-
-void Setup() {
-}
-
-void Teardown() {
-}
-```
-
-You must implement both `Setup` and `Teardown`, even if they are left empty.
-`Setup` is called before all tests and `Teardown` is called after all tests.
-Catch2 doesn't support global test setup and teardown so this is our only opportunity to manage the global life-cycle of tests.
-
-For unit tests, these methods should ideally be empty as true unit tests should run in isolation of any environment.
-Alternatively these methods can be used to mock components such as a HTTP server for example.
-
-For integration and functional tests we can spawn a local HTTP server or load a module in `Setup`, making sure to correctly clean up resources in `Teardown`.
-
-Add test files in `Private` and keep in mind the distinction between types of tests.
-
 ### Conventions
 
 **Files and folders**
@@ -226,11 +198,10 @@ TEST_CASE("Summary of test case", "[unit][feature][fast]")
 
 Tests can be built and run directly from Visual Studio on desktop platforms:
 1. Download the [Test Adapter for Catch2](https://marketplace.visualstudio.com/items?itemName=JohnnyHendriks.ext01) and install it.
-2. In Visual Studio go to *Test -> Configure Run Settings -> Select Solution Wide runsettings File* and select **Engine\Source\Programs\LowLevelTests\vstest.runsettings**.
-3. Build the test projects from Visual Studio to produce the executables. The test adapter will discover tests in the executables and they must contain the word **Test** e.g. *MyModuleTest.exe* or *MyModuleTests.exe*.
-4. The tests will be displayed in the test explorer: select *Test -> Test Explorer* from the menu: from here you can run tests and navigate to their source code.
+2. Build the test projects from Visual Studio to produce the executables. The test adapter will discover tests in the executables and they must contain the word **Test** e.g. *MyModuleTest.exe* or *MyModuleTests.exe*.
+3. The tests will be displayed in the test explorer: select *Test -> Test Explorer* from the menu: from here you can run tests and navigate to their source code.
 Running the tests will generate a *Test Detail Summary* and a Catch test report XML file in the same folder as the executable.
-5. If there are no tests in Test Explorer it's likely because the `Build` at step 3. didn't update the executable. Run a `Rebuild` on the test project and you should be good to go.
+4. If there are no tests in Test Explorer it's likely because the `Build` at step 3. didn't update the executable. Run a `Rebuild` on the test project and you should be good to go.
 
 > **If there are still no tests shown in the *Test Explorer*, perform a Rebuild on the test project and in the Output tab in VS select `Show output from: Tests` and look for `Catch2Adapter` in the log.
 > You should see a `Started Catch2Adapter test discovery...` line and a `Finished Catch2Adapter test discovery.`: between these lines you will see any problems reported by the Catch2 test adapter.**
@@ -246,7 +217,7 @@ Go to the binaries folder and run the executable with the extra flag `-- --wait`
 3. Run the BuildGraph command making sure to change the platform accordingly and set *-device* to the copied IP.
 In this example *-set:RunAudioTests=true* will run the audio tests while *-set:TestPlatformSwitch=true* means we are running on a Switch device..
 ```shell 
-.\RunUAT.bat BuildGraph -Script="Engine/Restricted/NotForLicensees/Build/LowLevelTests.xml" -set:RunAudioTests=true -set:TestPlatformSwitch=true -set:TestArgs="-deviceurl=https://horde.devtools.epicgames.com -device=<RESERVED SHARED DEVICE IP>" -Target="Run Low Level Tests"
+.\RunUAT.bat BuildGraph -Script="Engine/Build/LowLevelTests.xml" -set:RunAudioTests=true -set:TestPlatformSwitch=true -set:TestArgs="-deviceurl=https://horde.devtools.epicgames.com -device=<RESERVED SHARED DEVICE IP>" -Target="Run Low Level Tests"
 ```
 
 This will compile the test module, stage it and run the test app on the specified device.
