Fixed linker error for Android's build of LowLevelTests

diff --git a/Engine/Source/Programs/LowLevelTests/LowLevelTests.Target.cs b/Engine/Source/Programs/LowLevelTests/LowLevelTests.Target.cs
--- a/Engine/Source/Programs/LowLevelTests/LowLevelTests.Target.cs
+++ b/Engine/Source/Programs/LowLevelTests/LowLevelTests.Target.cs
@@ -1,5 +1,6 @@
 // Copyright Epic Games, Inc. All Rights Reserved.
 
+using System.IO;
 using UnrealBuildTool;
 
 [SupportedPlatforms(UnrealPlatformClass.All)]
@@ -9,5 +10,16 @@ public class LowLevelTestsTarget : TestTargetRules
 	{
 		bWithLowLevelTestsOverride = true;
 		IncludeOrderVersion = EngineIncludeOrderVersion.Latest;
+
+		if (Target.Platform == UnrealTargetPlatform.Android)
+		{
+			string VersionScriptFile = Path.GetTempPath() + "LLTWorkaroundScrip.ldscript";
+			using (StreamWriter Writer = File.CreateText(VersionScriptFile))
+			{
+				// Workaround for a linker bug when building LowLevelTests for Android
+				Writer.WriteLine("{ local: *; };");
+			}
+			AdditionalLinkerArguments = " -Wl,--version-script=\"" + VersionScriptFile + "\"";
+		}
 	}
 }
