fix for lowleveltests not running with interactive login

https://p4-swarm.epicgames.net/reviews/23796826

diff --git a/Engine/Source/Programs/LowLevelTests/LowLevelTests.Build.cs b/Engine/Source/Programs/LowLevelTests/LowLevelTests.Build.cs
--- a/Engine/Source/Programs/LowLevelTests/LowLevelTests.Build.cs
+++ b/Engine/Source/Programs/LowLevelTests/LowLevelTests.Build.cs
@@ -13,6 +13,13 @@ public class LowLevelTests : TestModuleRules
 				"CoreUObject"
 			});
 
+		if (Target.bBuildWithEditorOnlyData)
+		{
+			PrivateDependencyModuleNames.AddRange(
+				new string[] { "DesktopPlatform" }
+				);
+		}
+
 		UpdateBuildGraphPropertiesFile(new Metadata("Self"), false);
 	}
 }
\ No newline at end of file
