Rename LowLevelTests project and build/target to FoundationTests
#rnx

diff --git a/Engine/Build/LowLevelTests_GenProps.xml b/Engine/Build/LowLevelTests_GenProps.xml
--- a/Engine/Build/LowLevelTests_GenProps.xml
+++ b/Engine/Build/LowLevelTests_GenProps.xml
@@ -5,7 +5,7 @@
   <Property Name="HeadlessChaosDisabled" Value="False" />
   <Property Name="FoundationShort" Value="Foundation" />
   <Property Name="HeadlessChaosShort" Value="Headless Chaos" />
-  <Property Name="FoundationTarget" Value="LowLevelTests" />
+  <Property Name="FoundationTarget" Value="FoundationTests" />
   <Property Name="HeadlessChaosTarget" Value="HeadlessChaos" />
   <Property Name="FoundationBinariesRelative" Value="Engine\Binaries" />
   <Property Name="HeadlessChaosBinariesRelative" Value="Engine\Binaries" />
diff --git a/Engine/Source/Programs/LowLevelTests/FoundationTests.Build.cs b/Engine/Source/Programs/LowLevelTests/FoundationTests.Build.cs
new file mode 100644
--- /dev/null
+++ b/Engine/Source/Programs/LowLevelTests/FoundationTests.Build.cs
@@ -0,0 +1,26 @@
+// Copyright Epic Games, Inc. All Rights Reserved.
+using System;
+using UnrealBuildTool;
+
+public class FoundationTests : TestModuleRules
+{
+	public FoundationTests(ReadOnlyTargetRules Target) : base(Target, true)
+	{
+		PrivateDependencyModuleNames.AddRange(
+			new string[] {
+				"Core",
+				"Cbor",
+				"CoreUObject"
+			});
+
+		if (Target.bBuildWithEditorOnlyData)
+		{
+			PrivateDependencyModuleNames.AddRange(
+				new string[] {
+					"DesktopPlatform"
+				});
+		}
+
+		UpdateBuildGraphPropertiesFile(new Metadata("Foundation"), false);
+	}
+}
\ No newline at end of file
diff --git a/Engine/Source/Programs/LowLevelTests/FoundationTests.Target.cs b/Engine/Source/Programs/LowLevelTests/FoundationTests.Target.cs
new file mode 100644
--- /dev/null
+++ b/Engine/Source/Programs/LowLevelTests/FoundationTests.Target.cs
@@ -0,0 +1,25 @@
+// Copyright Epic Games, Inc. All Rights Reserved.
+
+using System.IO;
+using UnrealBuildTool;
+
+[SupportedPlatforms(UnrealPlatformClass.All)]
+public class FoundationTestsTarget : TestTargetRules
+{
+	public FoundationTestsTarget(TargetInfo Target) : base(Target)
+	{
+		// Collects all tests decorated with #if WITH_LOW_LEVELTESTS from dependencies
+		bWithLowLevelTestsOverride = true;
+
+		if (Target.Platform == UnrealTargetPlatform.Android)
+		{
+			string VersionScriptFile = Path.GetTempPath() + "LLTWorkaroundScrip.ldscript";
+			using (StreamWriter Writer = File.CreateText(VersionScriptFile))
+			{
+				// Workaround for a linker bug when building LowLevelTests for Android
+				Writer.WriteLine("{ local: *; };");
+			}
+			AdditionalLinkerArguments = " -Wl,--version-script=\"" + VersionScriptFile + "\"";
+		}
+	}
+}
diff --git a/Engine/Source/Programs/LowLevelTests/LowLevelTests.Build.cs b/Engine/Source/Programs/LowLevelTests/LowLevelTests.Build.cs
deleted file mode 100644
diff --git a/Engine/Source/Programs/LowLevelTests/LowLevelTests.Target.cs b/Engine/Source/Programs/LowLevelTests/LowLevelTests.Target.cs
deleted file mode 100644
