Foundation tests: with editor only data = true for desktop platforms, debug and development builds.

diff --git a/Engine/Source/Programs/LowLevelTests/FoundationTests.Target.cs b/Engine/Source/Programs/LowLevelTests/FoundationTests.Target.cs
--- a/Engine/Source/Programs/LowLevelTests/FoundationTests.Target.cs
+++ b/Engine/Source/Programs/LowLevelTests/FoundationTests.Target.cs
@@ -11,6 +11,9 @@ public class FoundationTestsTarget : TestTargetRules
 		// Collects all tests decorated with #if WITH_LOW_LEVELTESTS from dependencies
 		bWithLowLevelTestsOverride = true;
 
+		bBuildWithEditorOnlyData = Target.Platform.IsInGroup(UnrealPlatformGroup.Desktop)
+			&& (Target.Configuration == UnrealTargetConfiguration.Debug || Target.Configuration == UnrealTargetConfiguration.Development);
+
 		if (Target.Platform == UnrealTargetPlatform.Android)
 		{
 			string VersionScriptFile = Path.GetTempPath() + "LLTWorkaroundScrip.ldscript";
