Missing file for lowleveltests

diff --git a/Engine/Source/Programs/LowLevelTests/Private/TestObject.h b/Engine/Source/Programs/LowLevelTests/Private/TestObject.h
new file mode 100644
--- /dev/null
+++ b/Engine/Source/Programs/LowLevelTests/Private/TestObject.h
@@ -0,0 +1,53 @@
+// Copyright Epic Games, Inc. All Rights Reserved.
+
+#include "CoreTypes.h"
+#include "UObject/Object.h"
+
+#include "TestObject.generated.h"
+
+
+// Object to be used in text fixtures for various low level tests involving UObjects. 
+// Should contain various kinds of properties that can be used for various purposes depending on the test. 
+
+USTRUCT()
+struct FTestStruct
+{
+    GENERATED_BODY() 
+    
+    UPROPERTY()
+    TObjectPtr<UObject> StrongObjectReference;
+    
+    UPROPERTY()
+    TWeakObjectPtr<UObject> WeakObjectReference;
+
+    UPROPERTY()
+    TSoftObjectPtr<UObject> SoftObjectReference;
+    
+    UPROPERTY()
+    FSoftObjectPath SoftObjectPath;
+};
+
+UCLASS()
+class UTestObject : public UObject
+{
+    GENERATED_BODY()
+
+public:
+    UPROPERTY()
+    TObjectPtr<UObject> StrongObjectReference;
+    
+    UPROPERTY()
+    TWeakObjectPtr<UObject> WeakObjectReference;
+
+    UPROPERTY()
+    TSoftObjectPtr<UObject> SoftObjectReference;
+    
+    UPROPERTY()
+    FSoftObjectPath SoftObjectPath;
+    
+    UPROPERTY()
+    FTestStruct EmbeddedStruct;
+    
+    UPROPERTY()
+    TArray<FTestStruct> ArrayStructs;
+};
\ No newline at end of file
