Move the FoundationTests suite into their own folder under Source\Programs\LowLevelTests

diff --git a/Engine/Source/Programs/LowLevelTests/FoundationTests.Build.cs b/Engine/Source/Programs/LowLevelTests/FoundationTests.Build.cs
deleted file mode 100644
diff --git a/Engine/Source/Programs/LowLevelTests/FoundationTests.Target.cs b/Engine/Source/Programs/LowLevelTests/FoundationTests.Target.cs
deleted file mode 100644
diff --git a/Engine/Source/Programs/LowLevelTests/FoundationTests/FoundationTests.Build.cs b/Engine/Source/Programs/LowLevelTests/FoundationTests/FoundationTests.Build.cs
new file mode 100644
--- /dev/null
+++ b/Engine/Source/Programs/LowLevelTests/FoundationTests/FoundationTests.Build.cs
@@ -0,0 +1,60 @@
+// Copyright Epic Games, Inc. All Rights Reserved.
+using UnrealBuildTool;
+
+public class FoundationTests : TestModuleRules
+{
+	protected Metadata FoundationTestsMetadata = new Metadata() {
+		TestName = "Foundation",
+		TestShortName = "Foundation",
+		ReportType = "xml",
+		SupportedPlatforms = {
+			UnrealTargetPlatform.Win64,
+			UnrealTargetPlatform.Linux,
+			UnrealTargetPlatform.Mac,
+			UnrealTargetPlatform.Android,
+			UnrealTargetPlatform.IOS } };
+
+	/// <summary>
+	/// Test metadata to be used with BuildGraph
+	/// </summary>
+	public Metadata TestMetadata
+	{ 
+		get { return FoundationTestsMetadata; }
+	}
+
+	public FoundationTests(ReadOnlyTargetRules Target) : base(Target, true)
+	{
+		PrivateDependencyModuleNames.AddRange(
+			new string[] {
+				"Core",
+				"Cbor",
+				"CoreUObject",
+				"TelemetryUtils",
+				"AssetRegistry",
+				"ColorManagement"
+			});
+
+		if (Target.bBuildWithEditorOnlyData)
+		{
+			PrivateDependencyModuleNames.AddRange(
+				new string[] {
+					"DesktopPlatform"
+				});
+		}
+
+		string PlatformCompilationArgs;
+		foreach (var Platform in UnrealTargetPlatform.GetValidPlatforms())
+		{
+			if (Platform == UnrealTargetPlatform.Android)
+			{
+				PlatformCompilationArgs = "-allmodules -architectures=arm64";
+			}
+			else
+			{
+				PlatformCompilationArgs = "-allmodules";
+			}
+			TestMetadata.PlatformCompilationExtraArgs.Add(Platform, PlatformCompilationArgs);
+		}
+		UpdateBuildGraphPropertiesFile(TestMetadata);
+	}
+}
\ No newline at end of file
diff --git a/Engine/Source/Programs/LowLevelTests/FoundationTests/FoundationTests.Target.cs b/Engine/Source/Programs/LowLevelTests/FoundationTests/FoundationTests.Target.cs
new file mode 100644
--- /dev/null
+++ b/Engine/Source/Programs/LowLevelTests/FoundationTests/FoundationTests.Target.cs
@@ -0,0 +1,17 @@
+// Copyright Epic Games, Inc. All Rights Reserved.
+
+using UnrealBuildTool;
+
+[SupportedPlatforms(UnrealPlatformClass.All)]
+public class FoundationTestsTarget : TestTargetRules
+{
+	public FoundationTestsTarget(TargetInfo Target) : base(Target)
+	{
+		// Collects all tests decorated with #if WITH_LOW_LEVELTESTS from dependencies
+		bWithLowLevelTestsOverride = true;
+		bCompileWithPluginSupport = true;
+
+		bBuildWithEditorOnlyData = Target.Platform.IsInGroup(UnrealPlatformGroup.Desktop)
+			&& (Target.Configuration == UnrealTargetConfiguration.Debug || Target.Configuration == UnrealTargetConfiguration.Development);
+	}
+}
diff --git a/Engine/Source/Programs/LowLevelTests/FoundationTests/Private/TestObject.h b/Engine/Source/Programs/LowLevelTests/FoundationTests/Private/TestObject.h
new file mode 100644
--- /dev/null
+++ b/Engine/Source/Programs/LowLevelTests/FoundationTests/Private/TestObject.h
@@ -0,0 +1,53 @@
+// Copyright Epic Games, Inc. All Rights Reserved.
+
+#include "CoreTypes.h"
+#include "UObject/Object.h"
+
+#include "TestObject.generated.h"
+
+
+// Object to be used in text fixtures for various low level tests involving UObjects. 
+// Should contain various kinds of properties that can be used for various purposes depending on the test. 
+
+USTRUCT()
+struct FTestStruct
+{
+    GENERATED_BODY() 
+    
+    UPROPERTY()
+    TObjectPtr<UObject> StrongObjectReference;
+    
+    UPROPERTY()
+    TWeakObjectPtr<UObject> WeakObjectReference;
+
+    UPROPERTY()
+    TSoftObjectPtr<UObject> SoftObjectReference;
+    
+    UPROPERTY()
+    FSoftObjectPath SoftObjectPath;
+};
+
+UCLASS()
+class UTestObject : public UObject
+{
+    GENERATED_BODY()
+
+public:
+    UPROPERTY()
+    TObjectPtr<UObject> StrongObjectReference;
+    
+    UPROPERTY()
+    TWeakObjectPtr<UObject> WeakObjectReference;
+
+    UPROPERTY()
+    TSoftObjectPtr<UObject> SoftObjectReference;
+    
+    UPROPERTY()
+    FSoftObjectPath SoftObjectPath;
+    
+    UPROPERTY()
+    FTestStruct EmbeddedStruct;
+    
+    UPROPERTY()
+    TArray<FTestStruct> ArrayStructs;
+};
\ No newline at end of file
diff --git a/Engine/Source/Programs/LowLevelTests/FoundationTests/Tests/ArchiveReplaceObjectRefTests.cpp b/Engine/Source/Programs/LowLevelTests/FoundationTests/Tests/ArchiveReplaceObjectRefTests.cpp
new file mode 100644
--- /dev/null
+++ b/Engine/Source/Programs/LowLevelTests/FoundationTests/Tests/ArchiveReplaceObjectRefTests.cpp
@@ -0,0 +1,74 @@
+// Copyright Epic Games, Inc. All Rights Reserved.
+
+#include "CoreTypes.h"
+#include "Serialization/ArchiveReplaceObjectRef.h"
+#include "TestHarness.h"
+#include "TestObject.h"
+#include "UObject/Package.h"
+
+#if WITH_LOW_LEVEL_TESTS
+
+namespace UE
+{
+    TEST_CASE("CoreUObject::Serialization::FArchiveReplaceObjectRef", "[CoreUObject][Serialization]")
+    {
+        FName PackageName = MakeUniqueObjectName(nullptr, UPackage::StaticClass(), FName("/Memory/ArchiveReplaceObjectRefTest"));
+        UPackage* Package = NewObject<UPackage>(nullptr, UPackage::StaticClass(), PackageName);
+        UTestObject* Source = NewObject<UTestObject>(Package, "Source");
+        UTestObject* From1 = NewObject<UTestObject>(Package, "From1");
+        UTestObject* To1 = NewObject<UTestObject>(Package, "To1");
+        UTestObject* From2 = NewObject<UTestObject>(Package, "From2");
+        UTestObject* To2 = NewObject<UTestObject>(Package, "To2");
+
+        TMap<UObject*, UObject*> ReplacementMap;
+        ReplacementMap.Add(From1, To1);
+        ReplacementMap.Add(From2, To2);
+
+        SECTION("StrongReferences", "TObjectPtr fields should be replaced")
+        {
+            Source->StrongObjectReference = From1;
+            Source->EmbeddedStruct.StrongObjectReference = From2; 
+            FTestStruct& S = Source->ArrayStructs.AddZeroed_GetRef();
+            S.StrongObjectReference = From1;
+            FArchiveReplaceObjectRef Ar(Source, ReplacementMap);            
+            CHECK(Source->StrongObjectReference == To1);
+            CHECK(Source->EmbeddedStruct.StrongObjectReference == To2);
+            REQUIRE(Source->ArrayStructs.Num() == 1);
+            CHECK(Source->ArrayStructs[0].StrongObjectReference == To1);
+        }
+        
+        SECTION("WeakReferences", "TWeakObjectPtr fields should be replaced")
+        {
+            Source->WeakObjectReference = From1;
+            Source->EmbeddedStruct.WeakObjectReference = From2;
+            FTestStruct& S = Source->ArrayStructs.AddZeroed_GetRef();
+            S.WeakObjectReference = From1;
+            FArchiveReplaceObjectRef Ar(Source, ReplacementMap);            
+            CHECK(Source->WeakObjectReference == To1);
+            CHECK(Source->EmbeddedStruct.WeakObjectReference == To2);
+            REQUIRE(Source->ArrayStructs.Num() == 1);
+            CHECK(Source->ArrayStructs[0].WeakObjectReference == To1);
+        }
+        
+        SECTION("SoftReferences", "TSoftObjectPtr and FSoftObjectPath fields should be replaced")
+        {
+            Source->SoftObjectReference = From1;
+            Source->SoftObjectPath = From2;
+            Source->EmbeddedStruct.SoftObjectReference = From2;
+            Source->EmbeddedStruct.SoftObjectPath = From1;
+            FTestStruct& S = Source->ArrayStructs.AddZeroed_GetRef();
+            S.SoftObjectReference = From1;
+            S.SoftObjectPath = From2;
+            FArchiveReplaceObjectRef Ar(Source, ReplacementMap);            
+            CHECK(Source->SoftObjectReference.GetUniqueID() == FSoftObjectPath(To1));
+            CHECK(Source->SoftObjectPath == FSoftObjectPath(To2));
+            CHECK(Source->EmbeddedStruct.SoftObjectReference.GetUniqueID() == FSoftObjectPath(To2));
+            CHECK(Source->EmbeddedStruct.SoftObjectPath == FSoftObjectPath(To1));
+            REQUIRE(Source->ArrayStructs.Num() == 1);
+            CHECK(Source->ArrayStructs[0].SoftObjectReference.GetUniqueID() == FSoftObjectPath(To1));
+            CHECK(Source->ArrayStructs[0].SoftObjectPath == FSoftObjectPath(To2));
+        }
+    }
+}
+
+#endif
\ No newline at end of file
diff --git a/Engine/Source/Programs/LowLevelTests/FoundationTests/Tests/TestGroupEvents.cpp b/Engine/Source/Programs/LowLevelTests/FoundationTests/Tests/TestGroupEvents.cpp
new file mode 100644
--- /dev/null
+++ b/Engine/Source/Programs/LowLevelTests/FoundationTests/Tests/TestGroupEvents.cpp
@@ -0,0 +1,23 @@
+// Copyright Epic Games, Inc. All Rights Reserved.
+#include "TestCommon/CoreUtilities.h"
+#include "TestCommon/CoreUObjectUtilities.h"
+
+#include <catch2/catch_test_macros.hpp>
+
+GROUP_BEFORE_GLOBAL(Catch::DefaultGroup)
+{
+	InitTaskGraph();
+	InitThreadPool(true);
+#if WITH_COREUOBJECT
+	InitCoreUObject();
+#endif
+}
+
+GROUP_AFTER_GLOBAL(Catch::DefaultGroup)
+{
+#if WITH_COREUOBJECT
+	CleanupCoreUObject();
+#endif
+	CleanupThreadPool();
+	CleanupTaskGraph();
+}
diff --git a/Engine/Source/Programs/LowLevelTests/Private/TestObject.h b/Engine/Source/Programs/LowLevelTests/Private/TestObject.h
deleted file mode 100644
diff --git a/Engine/Source/Programs/LowLevelTests/Tests/ArchiveReplaceObjectRefTests.cpp b/Engine/Source/Programs/LowLevelTests/Tests/ArchiveReplaceObjectRefTests.cpp
deleted file mode 100644
diff --git a/Engine/Source/Programs/LowLevelTests/Tests/TestGroupEvents.cpp b/Engine/Source/Programs/LowLevelTests/Tests/TestGroupEvents.cpp
deleted file mode 100644
