Created MathCore module, a module meant to regroup all math utils that are not generic enough that they should be in Core (in order not to bloat it) or owned by the foundation team :
* Moved directed graph utils from audio to core, in namespace UE::MathCore::Graph
* Renamed Depth/BreadthFirstTraversal to Depth/BreadthFirstNodeTraversal and added Depth/BreadthFirstEdgeTraversal to traverse the graph edges
* Added FindLeaves method to retrieve graph nodes with no children from a starting node

Created MathCoreTests low-level test executable :
* Moved directed graph unit tests there
* Enabled on CIS on all platforms

#tests low-level tests

diff --git a/Engine/Build/LowLevelTests/MathCore.xml b/Engine/Build/LowLevelTests/MathCore.xml
new file mode 100644
--- /dev/null
+++ b/Engine/Build/LowLevelTests/MathCore.xml
@@ -0,0 +1,48 @@
+﻿<?xml version="1.0" encoding="utf-8"?>
+<BuildGraph xmlns:xsi="http://www.w3.org/2001/XMLSchema-instance" xsi:schemaLocation="http://www.epicgames.com/BuildGraph ../../Build/Graph/Schema.xsd" xmlns="http://www.epicgames.com/BuildGraph">
+  <Property Name="TestNames" Value="$(TestNames);MathCore" />
+  <Property Name="MathCoreDisabled" Value="False" />
+  <Property Name="MathCoreShort" Value="MathCore" />
+  <Property Name="MathCoreStagesWithProjectFile" Value="False" />
+  <Property Name="MathCoreTarget" Value="MathCoreTests" />
+  <Property Name="MathCoreBinariesRelative" Value="Engine\Binaries" />
+  <Property Name="MathCoreReportType" Value="xml" />
+  <Property Name="MathCoreGauntletArgs" Value="" />
+  <Property Name="MathCoreExtraArgs" />
+  <Property Name="MathCoreHasAfterSteps" Value="False" />
+  <Property Name="MathCoreUsesCatch2" Value="True" />
+  <Option Name="RunMathCoreTests" DefaultValue="False" Description="Run MathCore Tests" />
+  <Property Name="MathCoreSupportedPlatforms" Value="Win64;Linux;Mac;Android;IOS" />
+  <Property Name="Win64MathCoreTags" Value="" />
+  <Property Name="Win64MathCoreExtraCompilationArgs" Value="-allmodules" />
+  <Property Name="Win64MathCoreRunSupported" Value="True" />
+  <Property Name="Win64MathCoreRunContainerized" Value="False" />
+  <Property Name="MacMathCoreTags" Value="" />
+  <Property Name="MacMathCoreExtraCompilationArgs" Value="-allmodules" />
+  <Property Name="MacMathCoreRunSupported" Value="True" />
+  <Property Name="MacMathCoreRunContainerized" Value="False" />
+  <Property Name="IOSMathCoreTags" Value="" />
+  <Property Name="IOSMathCoreExtraCompilationArgs" Value="-allmodules" />
+  <Property Name="IOSMathCoreRunSupported" Value="False" />
+  <Property Name="IOSMathCoreRunContainerized" Value="False" />
+  <Property Name="AndroidMathCoreTags" Value="" />
+  <Property Name="AndroidMathCoreExtraCompilationArgs" Value="-allmodules -architectures=arm64" />
+  <Property Name="AndroidMathCoreRunSupported" Value="False" />
+  <Property Name="AndroidMathCoreRunContainerized" Value="False" />
+  <Property Name="LinuxMathCoreTags" Value="" />
+  <Property Name="LinuxMathCoreExtraCompilationArgs" Value="-allmodules" />
+  <Property Name="LinuxMathCoreRunSupported" Value="True" />
+  <Property Name="LinuxMathCoreRunContainerized" Value="False" />
+  <Property Name="LinuxArm64MathCoreTags" Value="" />
+  <Property Name="LinuxArm64MathCoreExtraCompilationArgs" Value="-allmodules" />
+  <Property Name="LinuxArm64MathCoreRunSupported" Value="True" />
+  <Property Name="LinuxArm64MathCoreRunContainerized" Value="False" />
+  <Property Name="TVOSMathCoreTags" Value="" />
+  <Property Name="TVOSMathCoreExtraCompilationArgs" Value="-allmodules" />
+  <Property Name="TVOSMathCoreRunSupported" Value="False" />
+  <Property Name="TVOSMathCoreRunContainerized" Value="False" />
+  <Property Name="VisionOSMathCoreTags" Value="" />
+  <Property Name="VisionOSMathCoreExtraCompilationArgs" Value="-allmodules" />
+  <Property Name="VisionOSMathCoreRunSupported" Value="False" />
+  <Property Name="VisionOSMathCoreRunContainerized" Value="False" />
+</BuildGraph>
\ No newline at end of file
diff --git a/Engine/Plugins/Runtime/Metasound/Source/MetasoundGraphCore/MetasoundGraphCore.Build.cs b/Engine/Plugins/Runtime/Metasound/Source/MetasoundGraphCore/MetasoundGraphCore.Build.cs
--- a/Engine/Plugins/Runtime/Metasound/Source/MetasoundGraphCore/MetasoundGraphCore.Build.cs
+++ b/Engine/Plugins/Runtime/Metasound/Source/MetasoundGraphCore/MetasoundGraphCore.Build.cs
@@ -16,6 +16,14 @@ namespace UnrealBuildTool.Rules
 					"AudioExtensions"
 				}
             );
+
+			PrivateDependencyModuleNames.AddRange(
+				new string[]
+				{
+					"MathCore",
+					// ... add private dependencies that you statically link with here ...	
+				}
+			);
 		}
 	}
 }
diff --git a/Engine/Plugins/Runtime/Metasound/Source/MetasoundGraphCore/Private/MetasoundGraphAlgo.cpp b/Engine/Plugins/Runtime/Metasound/Source/MetasoundGraphCore/Private/MetasoundGraphAlgo.cpp
--- a/Engine/Plugins/Runtime/Metasound/Source/MetasoundGraphCore/Private/MetasoundGraphAlgo.cpp
+++ b/Engine/Plugins/Runtime/Metasound/Source/MetasoundGraphCore/Private/MetasoundGraphAlgo.cpp
@@ -3,7 +3,7 @@
 #include "MetasoundGraphAlgo.h"
 #include "MetasoundGraphAlgoPrivate.h"
 
-#include "DSP/DirectedGraphAlgo.h"
+#include "Graph/DirectedGraphUtils.h"
 #include "Containers/Array.h"
 #include "Containers/Map.h"
 #include "Containers/Set.h"
@@ -21,7 +21,7 @@ namespace Metasound
 	class FDirectedGraphAlgoAdapter
 	{
 		using FNodePair = TTuple<const INode*, const INode*>;
-		using FDirectedEdge = Audio::FDirectedEdge;
+		using FDirectedEdge = UE::MathCore::Graph::FDirectedEdge;
 
 		// The TLazyObject holds the object and a flag to denote whether the 
 		// object has been initialized.
@@ -186,11 +186,11 @@ namespace Metasound
 			}
 
 			// Get in the form of a directed tree
-			const Audio::FDirectedTree& GetTree() const
+			const UE::MathCore::Graph::FDirectedTree& GetTree() const
 			{
 				if (!Lazy.Tree.bIsInitialized)
 				{
-					Audio::FDirectedGraphAlgo::BuildDirectedTree(GetUniqueEdges(), *Lazy.Tree);
+					UE::MathCore::Graph::BuildDirectedTree(GetUniqueEdges(), *Lazy.Tree);
 					Lazy.Tree.bIsInitialized = true;
 				}
 
@@ -198,11 +198,11 @@ namespace Metasound
 			}
 
 			// Get in the form of a transpose directed tree
-			const Audio::FDirectedTree& GetTransposeTree() const
+			const UE::MathCore::Graph::FDirectedTree& GetTransposeTree() const
 			{
 				if (!Lazy.TransposeTree.bIsInitialized)
 				{
-					Audio::FDirectedGraphAlgo::BuildTransposeDirectedTree(GetUniqueEdges(), *Lazy.TransposeTree);
+					UE::MathCore::Graph::BuildTransposeDirectedTree(GetUniqueEdges(), *Lazy.TransposeTree);
 					Lazy.TransposeTree.bIsInitialized = true;
 				}
 
@@ -214,9 +214,9 @@ namespace Metasound
 			struct FLazyCache
 			{
 				TLazyObject<TArray<int32>> UniqueVertices;
-				TLazyObject<TArray<Audio::FDirectedEdge>> UniqueEdges;
-				TLazyObject<Audio::FDirectedTree> Tree;
-				TLazyObject<Audio::FDirectedTree> TransposeTree;
+				TLazyObject<TArray<UE::MathCore::Graph::FDirectedEdge>> UniqueEdges;
+				TLazyObject<UE::MathCore::Graph::FDirectedTree> Tree;
+				TLazyObject<UE::MathCore::Graph::FDirectedTree> TransposeTree;
 			};
 
 			TSet<FDirectedEdge> EdgeSet;
@@ -321,7 +321,7 @@ namespace Metasound
 			TArray<int32> VertexOrder;
 
 			// Call algo implementation
-			bool bResult = Audio::FDirectedGraphAlgo::DepthFirstTopologicalSort(InAdapter.GetUniqueVertices(), InAdapter.GetUniqueEdges(), VertexOrder);
+			bool bResult = UE::MathCore::Graph::DepthFirstTopologicalSort(InAdapter.GetUniqueVertices(), InAdapter.GetUniqueEdges(), VertexOrder);
 
 			if (!bResult)
 			{
@@ -351,7 +351,7 @@ namespace Metasound
 			TArray<int32> VertexOrder;
 
 			// Call algo implementation
-			bool bResult = Audio::FDirectedGraphAlgo::KahnTopologicalSort(InAdapter.GetUniqueVertices(), InAdapter.GetUniqueEdges(), VertexOrder);
+			bool bResult = UE::MathCore::Graph::KahnTopologicalSort(InAdapter.GetUniqueVertices(), InAdapter.GetUniqueEdges(), VertexOrder);
 
 			if (!bResult)
 			{
@@ -376,14 +376,14 @@ namespace Metasound
 
 		bool TarjanStronglyConnectedComponents(const FDirectedGraphAlgoAdapter& InAdapter, TArray<FStronglyConnectedComponent>& OutComponents, bool bExcludeSingleVertex)
 		{
-			TArray<Audio::FStronglyConnectedComponent> StronglyConnectedComponents;
+			TArray<UE::MathCore::Graph::FStronglyConnectedComponent> StronglyConnectedComponents;
 
 			// Run tarjan on metasound derived graph edges 
-			if (Audio::FDirectedGraphAlgo::TarjanStronglyConnectedComponents(InAdapter.GetEdgeSet(), StronglyConnectedComponents, bExcludeSingleVertex))
+			if (UE::MathCore::Graph::TarjanStronglyConnectedComponents(InAdapter.GetEdgeSet(), StronglyConnectedComponents, bExcludeSingleVertex))
 			{
 				// If strongly connected components are found, they must be converted
 				// back into metasound types. 
-				for (const Audio::FStronglyConnectedComponent& Component : StronglyConnectedComponents)
+				for (const UE::MathCore::Graph::FStronglyConnectedComponent& Component : StronglyConnectedComponents)
 				{
 					FStronglyConnectedComponent& MetasoundGraphComponent = OutComponents.AddDefaulted_GetRef();
 
@@ -392,7 +392,7 @@ namespace Metasound
 						MetasoundGraphComponent.Nodes.Add(InAdapter.GetNode(Vertex));
 					}
 
-					for (const Audio::FDirectedEdge& Edge : Component.Edges)
+					for (const UE::MathCore::Graph::FDirectedEdge& Edge : Component.Edges)
 					{
 						TArray<FDataEdge> MetasoundComponentEdges;
 
@@ -421,7 +421,7 @@ namespace Metasound
 
 			for (int32 Vertex : InAdapter.GetUniqueInputVertices())
 			{
-				Audio::FDirectedGraphAlgo::DepthFirstTraversal(Vertex, InAdapter.GetTree(), [&](int32 InVertexBeingVisited) -> bool 
+				UE::MathCore::Graph::DepthFirstNodeTraversal(Vertex, InAdapter.GetTree(), [&](int32 InVertexBeingVisited) -> bool
 					{
 						bool bIsAlreadyInSet = false;
 						VisitedVertices.Add(InVertexBeingVisited, &bIsAlreadyInSet);
@@ -445,7 +445,7 @@ namespace Metasound
 
 			for (int32 Vertex : InAdapter.GetUniqueOutputVertices())
 			{ 
-				Audio::FDirectedGraphAlgo::DepthFirstTraversal(Vertex, InAdapter.GetTransposeTree(), [&](int32 InVertexBeingVisited) -> bool 
+				UE::MathCore::Graph::DepthFirstNodeTraversal(Vertex, InAdapter.GetTransposeTree(), [&](int32 InVertexBeingVisited) -> bool
 					{
 						bool bIsAlreadyInSet = false;
 						VisitedVertices.Add(InVertexBeingVisited, &bIsAlreadyInSet);
diff --git a/Engine/Source/Programs/LowLevelTests/MathCoreTests/MathCoreTests.Build.cs b/Engine/Source/Programs/LowLevelTests/MathCoreTests/MathCoreTests.Build.cs
new file mode 100644
--- /dev/null
+++ b/Engine/Source/Programs/LowLevelTests/MathCoreTests/MathCoreTests.Build.cs
@@ -0,0 +1,48 @@
+// Copyright Epic Games, Inc. All Rights Reserved.
+using UnrealBuildTool;
+
+public class MathCoreTests : TestModuleRules
+{
+	protected Metadata MathCoreTestsMetadata = new Metadata() {
+		TestName = "MathCore",
+		TestShortName = "MathCore",
+		ReportType = "xml",
+		SupportedPlatforms = {
+			UnrealTargetPlatform.Win64,
+			UnrealTargetPlatform.Linux,
+			UnrealTargetPlatform.Mac,
+			UnrealTargetPlatform.Android,
+			UnrealTargetPlatform.IOS } };
+
+	/// <summary>
+	/// Test metadata to be used with BuildGraph
+	/// </summary>
+	public Metadata TestMetadata
+	{ 
+		get { return MathCoreTestsMetadata; }
+	}
+
+	public MathCoreTests(ReadOnlyTargetRules Target) : base(Target, InUsesCatch2:true)
+	{
+		PrivateDependencyModuleNames.AddRange(
+			new string[] {
+				"Core",
+				"MathCore",
+			});
+
+		string PlatformCompilationArgs;
+		foreach (var Platform in UnrealTargetPlatform.GetValidPlatforms())
+		{
+			if (Platform == UnrealTargetPlatform.Android)
+			{
+				PlatformCompilationArgs = "-allmodules -architectures=arm64";
+			}
+			else
+			{
+				PlatformCompilationArgs = "-allmodules";
+			}
+			TestMetadata.PlatformCompilationExtraArgs.Add(Platform, PlatformCompilationArgs);
+		}
+		UpdateBuildGraphPropertiesFile(TestMetadata);
+	}
+}
\ No newline at end of file
diff --git a/Engine/Source/Programs/LowLevelTests/MathCoreTests/MathCoreTests.Target.cs b/Engine/Source/Programs/LowLevelTests/MathCoreTests/MathCoreTests.Target.cs
new file mode 100644
--- /dev/null
+++ b/Engine/Source/Programs/LowLevelTests/MathCoreTests/MathCoreTests.Target.cs
@@ -0,0 +1,15 @@
+// Copyright Epic Games, Inc. All Rights Reserved.
+
+using UnrealBuildTool;
+
+[SupportedPlatforms(UnrealPlatformClass.All)]
+public class MathCoreTestsTarget : TestTargetRules
+{
+	public MathCoreTestsTarget(TargetInfo Target) : base(Target)
+	{
+		// Collects all tests decorated with #if WITH_LOW_LEVELTESTS from dependencies
+		// TODO [jonathan.bard] : disabled for now : we can enable this (and move the tests in MathCoreTests to a Tests sub-folder alongside the code they're validating), when https://jira.it.epicgames.com/browse/UE-205189 is implemented : 
+		//  Without this, the tests from all linked modules (i.e. Core) would be run as part of this executable, which would be wasteful : 
+		// bWithLowLevelTestsOverride = true;
+	}
+}
diff --git a/Engine/Source/Programs/LowLevelTests/MathCoreTests/Private/Graph/DirectedGraphTest.cpp b/Engine/Source/Programs/LowLevelTests/MathCoreTests/Private/Graph/DirectedGraphTest.cpp
new file mode 100644
--- /dev/null
+++ b/Engine/Source/Programs/LowLevelTests/MathCoreTests/Private/Graph/DirectedGraphTest.cpp
@@ -0,0 +1,605 @@
+// Copyright Epic Games, Inc. All Rights Reserved.
+
+// TODO [jonathan.bard] : don't use "#if WITH_LOW_LEVEL_TESTS" for now : we can enable this (and move the tests in MathCoreTests to a Tests sub-folder alongside the code they're validating), when https://jira.it.epicgames.com/browse/UE-205189 is implemented : 
+//  Without this, the tests from all linked modules (i.e. Core) would be run as part of this executable, which would be wasteful : 
+//#if WITH_LOW_LEVEL_TESTS
+
+#include "CoreMinimal.h"
+#include "TestHarness.h"
+
+#include "Graph/DirectedGraphUtils.h"
+
+SCENARIO("MathCore::Graph::DirectedGraph", "[ApplicationContextMask][SmokeFilter]")
+{
+	using namespace UE::MathCore::Graph;
+
+	GIVEN("A graph without any strongly connected components")
+	{
+		TArray<int32> Vertices({0, 1, 2, 3, 4, 5, 6, 7});
+		TSet<FDirectedEdge> Edges;
+
+		Edges.Add(FDirectedEdge(0, 1));
+		Edges.Add(FDirectedEdge(0, 2));
+		Edges.Add(FDirectedEdge(1, 3));
+		Edges.Add(FDirectedEdge(2, 3));
+		Edges.Add(FDirectedEdge(3, 6));
+		Edges.Add(FDirectedEdge(0, 6));
+		Edges.Add(FDirectedEdge(6, 7));
+
+		WHEN("the tarjan algorithm tries to find strongly connected components")
+		{
+			TArray<FStronglyConnectedComponent> OutComponents;
+			bool bResult = TarjanStronglyConnectedComponents(Edges, OutComponents);
+
+			THEN("the algorithm does not find any strongly connected components")
+			{
+				REQUIRE(bResult == false);
+				REQUIRE(OutComponents.Num() == 0);
+			}
+		}
+
+		WHEN("the graph is sorted depth first")
+		{
+			TArray<int32> Order;
+			bool bSuccess = DepthFirstTopologicalSort(Vertices, Edges.Array(), Order);
+
+			THEN("the algorithm succeeds")
+			{
+				REQUIRE(bSuccess);
+			}
+
+			THEN("the algorithm produces a valid order")
+			{
+				int32 Index0 = Order.Find(0);
+				int32 Index1 = Order.Find(1);
+				int32 Index2 = Order.Find(2);
+				int32 Index3 = Order.Find(3);
+				int32 Index4 = Order.Find(4);
+				int32 Index5 = Order.Find(5);
+				int32 Index6 = Order.Find(6);
+				int32 Index7 = Order.Find(7);
+
+				REQUIRE(Index0 >= 0);
+				REQUIRE(Index1 >= 0);
+				REQUIRE(Index2 >= 0);
+				REQUIRE(Index3 >= 0);
+				REQUIRE(Index4 >= 0);
+				REQUIRE(Index5 >= 0);
+				REQUIRE(Index6 >= 0);
+				REQUIRE(Index7 >= 0);
+
+				REQUIRE(Index0 < Index1);
+				REQUIRE(Index0 < Index2);
+				REQUIRE(Index0 < Index3);
+				REQUIRE(Index0 < Index6);
+				REQUIRE(Index1 < Index3);
+				REQUIRE(Index1 < Index6);
+				REQUIRE(Index2 < Index3);
+				REQUIRE(Index2 < Index6);
+				REQUIRE(Index6 < Index7);
+			}
+		}
+
+		WHEN("the graph is sorted using kahns algorithm")
+		{
+			TArray<int32> Order;
+			bool bSuccess = KahnTopologicalSort(Vertices, Edges.Array(), Order);
+
+			THEN("the algorithm succeeds")
+			{
+				REQUIRE(bSuccess);
+			}
+
+			THEN("the algorithm produces a valid order")
+			{
+				int32 Index0 = Order.Find(0);
+				int32 Index1 = Order.Find(1);
+				int32 Index2 = Order.Find(2);
+				int32 Index3 = Order.Find(3);
+				int32 Index4 = Order.Find(4);
+				int32 Index5 = Order.Find(5);
+				int32 Index6 = Order.Find(6);
+				int32 Index7 = Order.Find(7);
+
+				REQUIRE(Index0 >= 0);
+				REQUIRE(Index1 >= 0);
+				REQUIRE(Index2 >= 0);
+				REQUIRE(Index3 >= 0);
+				REQUIRE(Index4 >= 0);
+				REQUIRE(Index5 >= 0);
+				REQUIRE(Index6 >= 0);
+				REQUIRE(Index7 >= 0);
+
+				REQUIRE(Index0 < Index1);
+				REQUIRE(Index0 < Index2);
+				REQUIRE(Index0 < Index3);
+				REQUIRE(Index0 < Index6);
+				REQUIRE(Index1 < Index3);
+				REQUIRE(Index1 < Index6);
+				REQUIRE(Index2 < Index3);
+				REQUIRE(Index2 < Index6);
+				REQUIRE(Index6 < Index7);
+			}
+		}
+
+		WHEN("the tree is traversed depth first for nodes")
+		{
+			FDirectedTree Tree;
+			BuildDirectedTree(Edges.Array(), Tree);
+
+			TMap<int32, int32> VisitOrder;
+			int32 VisitCounter = 0;
+
+			DepthFirstNodeTraversal(0, Tree, [&](int32 Vertex) -> bool
+				{
+					VisitOrder.Add(Vertex, VisitCounter);
+					VisitCounter++;
+					
+					if (Vertex == 6)
+					{
+						return false;
+					} 
+					return true;
+				}
+			);
+
+			THEN("The vertices are visited in a valid order")
+			{
+				REQUIRE(VisitOrder.Contains(0));
+				REQUIRE(VisitOrder.Contains(1));
+				REQUIRE(VisitOrder.Contains(2));
+				REQUIRE(VisitOrder.Contains(3));
+				REQUIRE(!VisitOrder.Contains(4));
+				REQUIRE(!VisitOrder.Contains(5));
+				REQUIRE(VisitOrder.Contains(6));
+				REQUIRE(!VisitOrder.Contains(7));
+
+				int32 Order0 = VisitOrder[0];
+				int32 Order1 = VisitOrder[1];
+				int32 Order2 = VisitOrder[2];
+				int32 Order3 = VisitOrder[3];
+				int32 Order6 = VisitOrder[6];
+
+				REQUIRE(Order0 >= 0);
+				REQUIRE(Order1 >= 0);
+				REQUIRE(Order2 >= 0);
+				REQUIRE(Order3 >= 0);
+				REQUIRE(Order6 >= 0);
+
+				REQUIRE(Order0 < Order1);
+				REQUIRE(Order0 < Order2);
+				REQUIRE(Order0 < Order3);
+				REQUIRE(Order0 < Order6);
+				bool bValidOrder = (Order1 < Order3) || (Order2 < Order3);
+				REQUIRE(bValidOrder);
+			}
+		}
+
+		WHEN("the tree is traversed depth first for edges")
+		{
+			FDirectedTree Tree;
+			BuildDirectedTree(Edges.Array(), Tree);
+
+			TMap<TTuple<int32, int32>, int32> VisitOrder;
+			int32 VisitCounter = 0;
+
+			DepthFirstEdgeTraversal(0, Tree, [&](int32 SourceVertex, int32 DestinationVertex) -> bool
+			{
+				VisitOrder.Add(MakeTuple(SourceVertex, DestinationVertex), VisitCounter);
+				VisitCounter++;
+
+				// Stop at one of the edges : 
+				if ((SourceVertex == 0) && (DestinationVertex == 1))
+				{
+					return false;
+				}
+				return true;
+			}
+			);
+
+			THEN("The vertices are visited in a valid order")
+			{
+				REQUIRE(VisitOrder.Contains(MakeTuple(0, 6)));
+				REQUIRE(VisitOrder.Contains(MakeTuple(6, 7)));
+				REQUIRE(VisitOrder.Contains(MakeTuple(0, 2)));
+				REQUIRE(VisitOrder.Contains(MakeTuple(2, 3)));
+				REQUIRE(VisitOrder.Contains(MakeTuple(3, 6)));
+				REQUIRE(VisitOrder.Contains(MakeTuple(0, 1)));
+				REQUIRE(!VisitOrder.Contains(MakeTuple(1, 3))); // This edge should have been skipped when we stopped at (0, 1)
+				REQUIRE(VisitOrder[MakeTuple(0, 6)] < VisitOrder[MakeTuple(6, 7)]);
+				REQUIRE(VisitOrder[MakeTuple(0, 2)] < VisitOrder[MakeTuple(2, 3)]);
+				REQUIRE(VisitOrder[MakeTuple(2, 3)] < VisitOrder[MakeTuple(3, 6)]);
+			}
+		}
+
+		WHEN("the tree is traversed breadth first")
+		{
+			FDirectedTree Tree;
+			BuildDirectedTree(Edges.Array(), Tree);
+
+			TMap<int32, int32> VisitOrder;
+			int32 VisitCounter = 0;
+
+			BreadthFirstNodeTraversal(0, Tree, [&](int32 Vertex) -> bool
+				{
+					VisitOrder.Add(Vertex, VisitCounter);
+					VisitCounter++;
+					
+					if (Vertex == 6)
+					{
+						return false;
+					} 
+					return true;
+				}
+			);
+
+			THEN("The vertices are visited in a valid order")
+			{
+				REQUIRE(VisitOrder.Contains(0));
+				REQUIRE(VisitOrder.Contains(1));
+				REQUIRE(VisitOrder.Contains(2));
+				REQUIRE(VisitOrder.Contains(3));
+				REQUIRE(!VisitOrder.Contains(4));
+				REQUIRE(!VisitOrder.Contains(5));
+				REQUIRE(VisitOrder.Contains(6));
+				REQUIRE(!VisitOrder.Contains(7));
+
+				int32 Order0 = VisitOrder[0];
+				int32 Order1 = VisitOrder[1];
+				int32 Order2 = VisitOrder[2];
+				int32 Order3 = VisitOrder[3];
+				int32 Order6 = VisitOrder[6];
+
+				REQUIRE(Order0 >= 0);
+				REQUIRE(Order1 >= 0);
+				REQUIRE(Order2 >= 0);
+				REQUIRE(Order3 >= 0);
+				REQUIRE(Order6 >= 0);
+
+				REQUIRE(Order0 < Order1);
+				REQUIRE(Order0 < Order2);
+				REQUIRE(Order0 < Order3);
+				REQUIRE(Order0 < Order6);
+				REQUIRE(Order1 < Order3);
+				REQUIRE(Order2 < Order3);
+				REQUIRE(Order6 < Order3);
+			}
+		}
+
+		WHEN("the tree is traversed breadth first for edges")
+		{
+			FDirectedTree Tree;
+			BuildDirectedTree(Edges.Array(), Tree);
+
+			TMap<TTuple<int32, int32>, int32> VisitOrder;
+			int32 VisitCounter = 0;
+
+			BreadthFirstEdgeTraversal(0, Tree, [&](int32 SourceVertex, int32 DestinationVertex) -> bool
+			{
+				VisitOrder.Add(MakeTuple(SourceVertex, DestinationVertex), VisitCounter);
+				VisitCounter++;
+
+				// Stop at one of the vertices : 
+				if (DestinationVertex == 6)
+				{
+					return false;
+				}
+				return true;
+			}
+			);
+
+			THEN("The vertices are visited in a valid order")
+			{
+				REQUIRE(VisitOrder.Contains(MakeTuple(0, 1)));
+				REQUIRE(VisitOrder.Contains(MakeTuple(0, 2)));
+				REQUIRE(VisitOrder.Contains(MakeTuple(0, 6)));
+				REQUIRE(VisitOrder.Contains(MakeTuple(1, 3)));
+				REQUIRE(VisitOrder.Contains(MakeTuple(2, 3))); 
+				REQUIRE(VisitOrder.Contains(MakeTuple(3, 6)));
+				REQUIRE(!VisitOrder.Contains(MakeTuple(6, 7))); // This edge should have been skipped when we stopped at vertex 6
+				REQUIRE(VisitOrder[MakeTuple(0, 1)] < VisitOrder[MakeTuple(1, 3)]);
+				REQUIRE(VisitOrder[MakeTuple(0, 2)] < VisitOrder[MakeTuple(2, 3)]);
+				REQUIRE(VisitOrder[MakeTuple(1, 3)] < VisitOrder[MakeTuple(3, 6)]);
+				REQUIRE(VisitOrder[MakeTuple(2, 3)] < VisitOrder[MakeTuple(3, 6)]);
+			}
+		}
+
+		WHEN("the tree leaves are retrieved")
+		{
+			FDirectedTree Tree;
+			TSet<FDirectedEdge> NewEdges(Edges);
+			// Add a couple more edges to test the child-less case : 
+			NewEdges.Add(FDirectedEdge(6, 8)); 
+			NewEdges.Add(FDirectedEdge(8, 9));
+			BuildDirectedTree(NewEdges.Array(), Tree);
+
+			// Manipulate the tree after creation to create a situation where there's a node without a child, which should still be considered as a leaf :
+			// Doing this will remove the edge from 8 to 9 while keeping node #8 : 
+			Tree.FindChecked(8).Children.Remove(9);
+			
+			TArray<int32> Leaves;
+			FindLeaves(0, Tree, Leaves);
+
+			THEN("All leaves are properly found")
+			{
+				REQUIRE(Leaves.Num() == 2);
+				REQUIRE(Leaves.Contains(7));
+				REQUIRE(Leaves.Contains(8));
+			}
+		}
+
+		WHEN("the transpose tree is traversed depth first")
+		{
+			FDirectedTree TransposeTree;
+			BuildTransposeDirectedTree(Edges.Array(), TransposeTree);
+
+			TMap<int32, int32> VisitOrder;
+			int32 VisitCounter = 0;
+
+			DepthFirstNodeTraversal(7, TransposeTree, [&](int32 Vertex) -> bool
+				{
+					VisitOrder.Add(Vertex, VisitCounter);
+					VisitCounter++;
+					
+					if (Vertex == 3)
+					{
+						return false;
+					} 
+					return true;
+				}
+			);
+
+			THEN("The vertices are visited in a valid order")
+			{
+				REQUIRE(VisitOrder.Contains(0));
+				REQUIRE(!VisitOrder.Contains(1));
+				REQUIRE(!VisitOrder.Contains(2));
+				REQUIRE(VisitOrder.Contains(3));
+				REQUIRE(!VisitOrder.Contains(4));
+				REQUIRE(!VisitOrder.Contains(5));
+				REQUIRE(VisitOrder.Contains(6));
+				REQUIRE(VisitOrder.Contains(7));
+
+				int32 Order0 = VisitOrder[0];
+				int32 Order3 = VisitOrder[3];
+				int32 Order6 = VisitOrder[6];
+				int32 Order7 = VisitOrder[7];
+
+				REQUIRE(Order0 >= 0);
+				REQUIRE(Order3 >= 0);
+				REQUIRE(Order6 >= 0);
+				REQUIRE(Order7 >= 0);
+
+				REQUIRE(Order7 < Order0);
+				REQUIRE(Order7 < Order3);
+				REQUIRE(Order7 < Order6);
+				REQUIRE(Order6 < Order0);
+				REQUIRE(Order6 < Order3);
+			}
+		}
+
+		WHEN("the transpose tree is traversed depth first for edges")
+		{
+			FDirectedTree Tree;
+			BuildTransposeDirectedTree(Edges.Array(), Tree);
+
+			TMap<TTuple<int32, int32>, int32> VisitOrder;
+			int32 VisitCounter = 0;
+
+			DepthFirstEdgeTraversal(7, Tree, [&](int32 SourceVertex, int32 DestinationVertex) -> bool
+			{
+				VisitOrder.Add(MakeTuple(SourceVertex, DestinationVertex), VisitCounter);
+				VisitCounter++;
+
+				// Stop at one of the edges : 
+				if ((SourceVertex == 3) && (DestinationVertex == 1))
+				{
+					return false;
+				}
+				return true;
+			}
+			);
+
+			THEN("The vertices are visited in a valid order")
+			{
+				REQUIRE(VisitOrder.Contains(MakeTuple(7, 6)));
+				REQUIRE(VisitOrder.Contains(MakeTuple(6, 0)));
+				REQUIRE(VisitOrder.Contains(MakeTuple(6, 3)));
+				REQUIRE(VisitOrder.Contains(MakeTuple(3, 2)));
+				REQUIRE(VisitOrder.Contains(MakeTuple(2, 0)));
+				REQUIRE(VisitOrder.Contains(MakeTuple(3, 1)));
+				REQUIRE(!VisitOrder.Contains(MakeTuple(1, 0))); // This edge should have been skipped when we stopped at (3, 1)
+				REQUIRE(VisitOrder[MakeTuple(7, 6)] == 0);
+				REQUIRE(VisitOrder[MakeTuple(6, 0)] < VisitOrder[MakeTuple(3, 1)]);
+				REQUIRE(VisitOrder[MakeTuple(6, 0)] < VisitOrder[MakeTuple(3, 2)]);
+				REQUIRE(VisitOrder[MakeTuple(6, 3)] < VisitOrder[MakeTuple(3, 1)]);
+				REQUIRE(VisitOrder[MakeTuple(6, 3)] < VisitOrder[MakeTuple(3, 2)]);
+				REQUIRE(VisitOrder[MakeTuple(3, 2)] < VisitOrder[MakeTuple(2, 0)]);
+			}
+		}
+
+		WHEN("the transpose tree is traversed breadth first")
+		{
+			FDirectedTree TransposeTree;
+			BuildTransposeDirectedTree(Edges.Array(), TransposeTree);
+
+			TMap<int32, int32> VisitOrder;
+			int32 VisitCounter = 0;
+
+			BreadthFirstNodeTraversal(7, TransposeTree, [&](int32 Vertex) -> bool
+				{
+					VisitOrder.Add(Vertex, VisitCounter);
+					VisitCounter++;
+					
+					if (Vertex == 3)
+					{
+						return false;
+					} 
+					return true;
+				}
+			);
+
+			THEN("The vertices are visited in a valid order")
+			{
+				REQUIRE(VisitOrder.Contains(0));
+				REQUIRE(!VisitOrder.Contains(1));
+				REQUIRE(!VisitOrder.Contains(2));
+				REQUIRE(VisitOrder.Contains(3));
+				REQUIRE(!VisitOrder.Contains(4));
+				REQUIRE(!VisitOrder.Contains(5));
+				REQUIRE(VisitOrder.Contains(6));
+				REQUIRE(VisitOrder.Contains(7));
+
+				int32 Order0 = VisitOrder[0];
+				int32 Order3 = VisitOrder[3];
+				int32 Order6 = VisitOrder[6];
+				int32 Order7 = VisitOrder[7];
+
+				REQUIRE(Order0 >= 0);
+				REQUIRE(Order3 >= 0);
+				REQUIRE(Order6 >= 0);
+				REQUIRE(Order7 >= 0);
+
+				REQUIRE(Order7 < Order0);
+				REQUIRE(Order7 < Order3);
+				REQUIRE(Order7 < Order6);
+				REQUIRE(Order6 < Order0);
+				REQUIRE(Order6 < Order3);
+			}
+
+		}
+
+		WHEN("the transpose tree is traversed breadth first for edges")
+		{
+			FDirectedTree Tree;
+			BuildTransposeDirectedTree(Edges.Array(), Tree);
+
+			TMap<TTuple<int32, int32>, int32> VisitOrder;
+			int32 VisitCounter = 0;
+
+			BreadthFirstEdgeTraversal(7, Tree, [&](int32 SourceVertex, int32 DestinationVertex) -> bool
+			{
+				VisitOrder.Add(MakeTuple(SourceVertex, DestinationVertex), VisitCounter);
+				VisitCounter++;
+
+				// Stop at one of the edges : 
+				if ((SourceVertex == 3) && (DestinationVertex == 1))
+				{
+					return false;
+				}
+				return true;
+			}
+			);
+
+			THEN("The vertices are visited in a valid order")
+			{
+				REQUIRE(VisitOrder.Contains(MakeTuple(7, 6)));
+				REQUIRE(VisitOrder.Contains(MakeTuple(6, 3)));
+				REQUIRE(VisitOrder.Contains(MakeTuple(6, 0)));
+				REQUIRE(VisitOrder.Contains(MakeTuple(3, 2)));
+				REQUIRE(VisitOrder.Contains(MakeTuple(3, 1)));
+				REQUIRE(VisitOrder.Contains(MakeTuple(2, 0)));
+				REQUIRE(!VisitOrder.Contains(MakeTuple(1, 0))); // This edge should have been skipped when we stopped at (3, 1)
+				REQUIRE(VisitOrder[MakeTuple(7, 6)] == 0);
+				REQUIRE(VisitOrder[MakeTuple(6, 3)] < VisitOrder[MakeTuple(3, 1)]);
+				REQUIRE(VisitOrder[MakeTuple(6, 0)] < VisitOrder[MakeTuple(3, 1)]);
+				REQUIRE(VisitOrder[MakeTuple(6, 3)] < VisitOrder[MakeTuple(3, 2)]);
+				REQUIRE(VisitOrder[MakeTuple(6, 0)] < VisitOrder[MakeTuple(3, 2)]);
+				REQUIRE(VisitOrder[MakeTuple(3, 1)] < VisitOrder[MakeTuple(2, 0)]);
+			}
+		}
+		
+		WHEN("the transpose tree leaves are retrieved")
+		{
+			FDirectedTree Tree;
+			BuildTransposeDirectedTree(Edges.Array(), Tree);
+
+			TArray<int32> Leaves;
+			FindLeaves(0, Tree, Leaves);
+
+			THEN("All leaves are properly found")
+			{
+				REQUIRE(Leaves.Num() == 1);
+				REQUIRE(Leaves[0] == 0);
+			}
+		}
+	}
+
+	GIVEN("A graph with strongly connected components")
+	{
+		TSet<FDirectedEdge> Edges;
+
+		Edges.Add(FDirectedEdge(0, 1));
+		Edges.Add(FDirectedEdge(0, 2));
+		Edges.Add(FDirectedEdge(1, 3));
+		Edges.Add(FDirectedEdge(2, 3));
+		Edges.Add(FDirectedEdge(3, 6));
+		Edges.Add(FDirectedEdge(6, 0));
+		Edges.Add(FDirectedEdge(6, 7));
+		Edges.Add(FDirectedEdge(7, 8));
+		Edges.Add(FDirectedEdge(8, 9));
+		Edges.Add(FDirectedEdge(9, 8));
+
+
+		WHEN("the graph is analyzed")
+		{
+			TArray<FStronglyConnectedComponent> OutComponents;
+			bool bResult = TarjanStronglyConnectedComponents(Edges, OutComponents);
+
+			THEN("strongly connected components are found")
+			{
+				REQUIRE(bResult == true);
+				REQUIRE(OutComponents.Num() == 2);
+			}
+
+			THEN("the strongly connected components are correct")
+			{
+				// Components are always reported in reverse order using the
+				// tarjan algorithm. This first requirement is checking for the
+				// last strongly connected component in the graph.
+				FStronglyConnectedComponent& Comp = OutComponents[0];
+
+				REQUIRE(Comp.Vertices.Num() == 2);
+				REQUIRE(Comp.Vertices.Contains(8));
+				REQUIRE(Comp.Vertices.Contains(9));
+
+				REQUIRE(Comp.Edges.Num() == 2);
+				REQUIRE(Comp.Edges.Contains(FDirectedEdge(8, 9)));
+				REQUIRE(Comp.Edges.Contains(FDirectedEdge(9, 8)));
+
+				Comp = OutComponents[1];
+
+				REQUIRE(Comp.Vertices.Num() == 5);
+				REQUIRE(Comp.Vertices.Contains(0));
+				REQUIRE(Comp.Vertices.Contains(1));
+				REQUIRE(Comp.Vertices.Contains(2));
+				REQUIRE(Comp.Vertices.Contains(3));
+				REQUIRE(Comp.Vertices.Contains(6));
+
+				REQUIRE(Comp.Edges.Num() == 6);
+				REQUIRE(Comp.Edges.Contains(FDirectedEdge(0, 1)));
+				REQUIRE(Comp.Edges.Contains(FDirectedEdge(0, 2)));
+				REQUIRE(Comp.Edges.Contains(FDirectedEdge(1, 3)));
+				REQUIRE(Comp.Edges.Contains(FDirectedEdge(2, 3)));
+				REQUIRE(Comp.Edges.Contains(FDirectedEdge(3, 6)));
+				REQUIRE(Comp.Edges.Contains(FDirectedEdge(6, 0)));
+			}
+		}
+
+		WHEN("the tree leaves are retrieved")
+		{
+			FDirectedTree Tree;
+			BuildTransposeDirectedTree(Edges.Array(), Tree);
+
+			TArray<int32> Leaves;
+			FindLeaves(0, Tree, Leaves);
+
+			THEN("No leaves are found")
+			{
+				REQUIRE(Leaves.IsEmpty());
+			}
+		}
+	}
+}
+
+//#endif // WITH_LOW_LEVEL_TESTS
\ No newline at end of file
diff --git a/Engine/Source/Runtime/MathCore/MathCore.Build.cs b/Engine/Source/Runtime/MathCore/MathCore.Build.cs
new file mode 100644
--- /dev/null
+++ b/Engine/Source/Runtime/MathCore/MathCore.Build.cs
@@ -0,0 +1,19 @@
+// Copyright Epic Games, Inc. All Rights Reserved.
+
+using UnrealBuildTool;
+
+public class MathCore : ModuleRules
+{
+	public MathCore(ReadOnlyTargetRules Target) : base(Target)
+	{
+		PCHUsage = ModuleRules.PCHUsageMode.UseExplicitOrSharedPCHs;
+
+		PublicDependencyModuleNames.AddRange(
+			new string[] {
+				"Core"
+			}
+		);
+
+		UnsafeTypeCastWarningLevel = WarningLevel.Error;
+	}
+}
diff --git a/Engine/Source/Runtime/MathCore/Private/Graph/DirectedGraphUtils.cpp b/Engine/Source/Runtime/MathCore/Private/Graph/DirectedGraphUtils.cpp
new file mode 100644
--- /dev/null
+++ b/Engine/Source/Runtime/MathCore/Private/Graph/DirectedGraphUtils.cpp
@@ -0,0 +1,562 @@
+// Copyright Epic Games, Inc. All Rights Reserved.
+
+#include "Graph/DirectedGraphUtils.h"
+
+namespace UE::MathCore::Graph::Private
+{
+	// Used to track which vertices have been visited and vertex children.
+	struct FDepthFirstSortInfo
+	{
+		bool TemporaryMark = false;
+		bool PermanentMark = false;
+		TArray<int32> Children;
+	};
+
+	// Recursive function for depth first topological sorting.
+	bool DepthFirstTopologicalSortVisit(int32 InVertex, TMap<int32, FDepthFirstSortInfo>& InSortInfos, TArray<int32>& OutReverseOrder)
+	{
+		FDepthFirstSortInfo& Info = InSortInfos[InVertex];
+
+		if (Info.PermanentMark)
+		{
+			// Exit depth first traversal
+			return true;
+		}
+		else if (Info.TemporaryMark)
+		{
+			// Found a cycle in the DAG.
+			return false;
+		}
+
+		Info.TemporaryMark = true;
+
+		for (int32 Child : Info.Children)
+		{
+			if (!DepthFirstTopologicalSortVisit(Child, InSortInfos, OutReverseOrder))
+			{
+				// cycle found in children.
+				return false;
+			}
+		}
+
+		Info.TemporaryMark = false;
+		Info.PermanentMark = true;
+
+		OutReverseOrder.Add(InVertex);
+
+		// subgraph successfully sorted.
+		return true;
+	}
+
+	struct FTarjanVertexInfo
+	{
+		// Tag value given to a vertex.
+		int32 Tag = INDEX_NONE;
+
+		// Lowest tag in strongly connected components
+		int32 LowLinkTag = INDEX_NONE;
+
+		// Flag for global stack
+		bool bIsOnStack = false;
+
+		// Immediate children of vertex in DAG.
+		TArray<int32> Children;
+	};
+
+	struct FTarjanAlgoImplSettings
+	{
+		// Single vertices can be returned as strongly connected components
+		// but usually they are not wanted by the caller. Setting this to
+		// true will require all strongly connected components to contain
+		// two or more vertices.
+		bool bExcludeSingleVertex;
+	};
+
+	// Private implementation of algorithm.
+	class FTarjanAlgoImpl
+	{
+	public:
+
+		FTarjanAlgoImpl(const FTarjanAlgoImplSettings& InSettings)
+			: Settings(InSettings)
+		{
+		}
+
+		bool FindStronglyConnectedComponents(const TSet<FDirectedEdge>& InEdges, TArray<FStronglyConnectedComponent>& OutComponents)
+		{
+			using FElement = TMap<int32, FTarjanVertexInfo>::ElementType;
+
+			// Initialize internal data stores and tracking info.
+			Init(InEdges);
+
+			int32 OriginalNumOutComponents = OutComponents.Num();
+
+			// Iterate through all vertices. Update(...) does a 
+			// depth-first-search which may result some vertices getting 
+			// inspected during recursive calls to Update(...).
+			for (FElement& Element : VertexInfos)
+			{
+				// Only examine vertice if it hasn't already been inspected.
+				if (Element.Value.Tag == INDEX_NONE)
+				{
+					Update(Element.Key, Element.Value, OutComponents);
+				}
+			}
+
+			// Check if the size of OutComponents has grown.
+			return OutComponents.Num() > OriginalNumOutComponents;
+		}
+
+	private:
+
+		void Init(const TSet<FDirectedEdge>& InEdges)
+		{
+			CurrentTag = 0;
+
+			VertexInfos.Reset();
+			VertexStack.Reset();
+
+			for (const FDirectedEdge& Edge : InEdges)
+			{
+				FTarjanVertexInfo& Info = VertexInfos.FindOrAdd(Edge.Get<0>());
+				Info.Children.AddUnique(Edge.Get<1>());
+
+				VertexInfos.FindOrAdd(Edge.Get<1>());
+			}
+		}
+
+		void Update(int32 InVertex, FTarjanVertexInfo& InVertexInfo, TArray<FStronglyConnectedComponent>& OutComponents)
+		{
+			InVertexInfo.Tag = CurrentTag;
+			InVertexInfo.LowLinkTag = CurrentTag;
+			InVertexInfo.bIsOnStack = true;
+			VertexStack.Add(InVertex);
+
+			// Tag needs to always be increasing during subsequent calls to
+			// Update(...).
+			CurrentTag++;
+
+			// Perform depth first traversal through the graph looking for 
+			// children that are already on the stack (which denotes a 
+			// strongly connected component)
+			for (int32 ChildVertex : InVertexInfo.Children)
+			{
+				FTarjanVertexInfo& ChildInfo = VertexInfos[ChildVertex];
+
+				if (ChildInfo.Tag == INDEX_NONE)
+				{
+					// If ChildInfo.Tag == INDEX_NONE, then this child is 
+					// not on the stack because it has never been inspected.
+
+					// Continue in depth-first-traversal. 
+					Update(ChildVertex, ChildInfo, OutComponents);
+
+					// Update to use the lowest link tag found in children.
+					InVertexInfo.LowLinkTag = FMath::Min(InVertexInfo.LowLinkTag, ChildInfo.LowLinkTag);
+				}
+				else if (ChildInfo.bIsOnStack)
+				{
+					// This child is on the stack and represents a strongly 
+					// connected component. 
+					InVertexInfo.LowLinkTag = FMath::Min(InVertexInfo.LowLinkTag, ChildInfo.Tag);
+				}
+			}
+
+			if (InVertexInfo.LowLinkTag == InVertexInfo.Tag)
+			{
+				// Create a strongly connected component object if this is 
+				// the root vertex of a strongly connected component.
+				AddStronglyConnectedComponent(InVertex, OutComponents);
+			}
+		}
+
+		// Adds the current strongly connected component to the 
+		// OutComponents array.
+		void AddStronglyConnectedComponent(int32 InRootVertex, TArray<FStronglyConnectedComponent>& OutComponents)
+		{
+			// When a single vertex represents the entire strongly connected
+			// component, the root vertex is the same as the vertex at the
+			// top of the stack.
+			bool bIsSingleVertex = (VertexStack.Top() == InRootVertex);
+
+			if (bIsSingleVertex && Settings.bExcludeSingleVertex)
+			{
+				// This strongly connected component is skipped because it 
+				// consists of a single vertex. 
+				//
+				// The stack still needs to be updated as well as the vertex
+				// info for the root vertex.
+				int32 TopVertex = VertexStack.Pop();
+				FTarjanVertexInfo& TopInfo = VertexInfos[TopVertex];
+				TopInfo.bIsOnStack = false;
+			}
+			else
+			{
+				bool bDidCompleteLoop = false;
+
+				// Create a new strongly connected component subgraph.
+				FStronglyConnectedComponent& StronglyConnectedComponent = OutComponents.AddDefaulted_GetRef();
+
+				// Pop off the top of the vertex stack until we reach the
+				// root vertex.
+				while (!bDidCompleteLoop)
+				{
+					int32 TopVertex = VertexStack.Pop();
+
+					FTarjanVertexInfo& TopInfo = VertexInfos[TopVertex];
+					TopInfo.bIsOnStack = false;
+
+					StronglyConnectedComponent.Vertices.Add(TopVertex);
+
+					// The unwinding continues until the current vertex is 
+					// removed from the stack.
+					bDidCompleteLoop = (TopVertex == InRootVertex);
+				}
+
+				// Add the edges associated with the strongly connected
+				// component. 
+				for (int32 ComponentVertex : StronglyConnectedComponent.Vertices)
+				{
+					FTarjanVertexInfo& ComponentInfo = VertexInfos[ComponentVertex];
+
+					for (int32 ChildVertex : ComponentInfo.Children)
+					{
+						if (StronglyConnectedComponent.Vertices.Contains(ChildVertex))
+						{
+							StronglyConnectedComponent.Edges.Emplace(ComponentVertex, ChildVertex);
+						}
+					}
+				}
+			}
+		}
+
+		FTarjanAlgoImplSettings Settings;
+
+		int32 CurrentTag = 0;
+
+		TMap<int32, FTarjanVertexInfo> VertexInfos;
+
+		TArray<int32> VertexStack;
+	};
+} // namespace namespace UE::MathCore::Graph::Private
+
+namespace UE::MathCore::Graph
+{
+	void BuildDirectedTree(TArrayView<const FDirectedEdge> InEdges, FDirectedTree& OutTree)
+	{
+		for (const FDirectedEdge& Edge : InEdges)
+		{
+			FDirectedTreeElement& Element = OutTree.FindOrAdd(Edge.Get<0>());
+
+			Element.Children.AddUnique(Edge.Get<1>());
+		}
+	}
+
+	void BuildTransposeDirectedTree(TArrayView<const FDirectedEdge> InEdges, FDirectedTree& OutTree)
+	{
+		for (const FDirectedEdge& Edge : InEdges)
+		{
+			FDirectedTreeElement& Element = OutTree.FindOrAdd(Edge.Get<1>());
+
+			Element.Children.AddUnique(Edge.Get<0>());
+		}
+	}
+
+	void DepthFirstNodeTraversal(int32 InInitialVertex, const FDirectedTree& InTree, TFunctionRef<bool(int32)> InVisitFunc)
+	{
+		// Non recursive depth first traversal.
+
+		TArray<int32> Stack({InInitialVertex});
+		TSet<int32> Visited;
+
+		bool bVisitChildren = true;
+
+		while (Stack.Num() > 0)
+		{
+			int32 CurrentVertex = Stack.Pop();
+
+			if (Visited.Contains(CurrentVertex))
+			{
+				continue;
+			}
+
+			bVisitChildren = InVisitFunc(CurrentVertex);
+
+			Visited.Add(CurrentVertex);
+
+			if (bVisitChildren)
+			{
+				if (const FDirectedTreeElement* Element = InTree.Find(CurrentVertex))
+				{
+					Stack.Append(Element->Children);
+				}
+			}
+		}
+	}
+
+	void DepthFirstEdgeTraversal(int32 InInitialVertex, const FDirectedTree& InTree, TFunctionRef<bool(int32, int32)> InVisitFunc)
+	{
+		// Non recursive depth first traversal.
+
+		TArray<TTuple<int32, int32>> Stack({ MakeTuple(INDEX_NONE, InInitialVertex) });
+		TSet<TTuple<int32, int32>> Visited;
+
+		bool bVisitChildren = true;
+
+		while (Stack.Num() > 0)
+		{
+			TTuple<int32, int32> CurrentEdge = Stack.Pop();
+
+			if (Visited.Contains(CurrentEdge))
+			{
+				continue;
+			}
+
+			// The first element is not an edge so don't visit it : 
+			if (CurrentEdge.Get<0>() != INDEX_NONE)
+			{
+				bVisitChildren = InVisitFunc(CurrentEdge.Get<0>(), CurrentEdge.Get<1>());
+			}
+
+			Visited.Add(CurrentEdge);
+
+			if (bVisitChildren)
+			{
+				if (const FDirectedTreeElement* Element = InTree.Find(CurrentEdge.Get<1>()))
+				{
+					TArray<TTuple<int32, int32>> ChildEdges;
+					ChildEdges.Reserve(Element->Children.Num());
+					for (int32 Child : Element->Children)
+					{
+						ChildEdges.Add(MakeTuple(CurrentEdge.Get<1>(), Child));
+					}
+					Stack.Append(MoveTemp(ChildEdges));
+				}
+			}
+		}
+	}
+
+	void BreadthFirstNodeTraversal(int32 InInitialVertex, const FDirectedTree& InTree, TFunctionRef<bool(int32)> InVisitFunc)
+	{
+		// Non recursive breadth first traversal.
+
+		TArray<int32> Queue({InInitialVertex});
+		TSet<int32> Visited;
+
+		bool bVisitChildren = true;
+
+		while (Queue.Num() > 0)
+		{
+			int32 CurrentVertex = Queue[0];
+			Queue.RemoveAt(0);
+
+			if (Visited.Contains(CurrentVertex))
+			{
+				continue;
+			}
+
+			bVisitChildren = InVisitFunc(CurrentVertex);
+
+			Visited.Add(CurrentVertex);
+
+			if (bVisitChildren)
+			{
+				if (const FDirectedTreeElement* Element = InTree.Find(CurrentVertex))
+				{
+					Queue.Append(Element->Children);
+				}
+			}
+		}
+	}
+
+	void BreadthFirstEdgeTraversal(int32 InInitialVertex, const FDirectedTree& InTree, TFunctionRef<bool(int32, int32)> InVisitFunc)
+	{
+		// Non recursive breadth first traversal.
+
+		TArray<TTuple<int32, int32>> Queue({ MakeTuple(INDEX_NONE, InInitialVertex) });
+		TSet<TTuple<int32, int32>> Visited;
+
+		bool bVisitChildren = true;
+
+		while (Queue.Num() > 0)
+		{
+			TTuple<int32, int32> CurrentEdge = Queue[0];
+			Queue.RemoveAt(0);
+
+			if (Visited.Contains(CurrentEdge))
+			{
+				continue;
+			}
+
+			// The first element is not an edge so don't visit it : 
+			if (CurrentEdge.Get<0>() != INDEX_NONE)
+			{
+				bVisitChildren = InVisitFunc(CurrentEdge.Get<0>(), CurrentEdge.Get<1>());
+			}
+
+			Visited.Add(CurrentEdge);
+
+			if (bVisitChildren)
+			{
+				if (const FDirectedTreeElement* Element = InTree.Find(CurrentEdge.Get<1>()))
+				{
+					TArray<TTuple<int32, int32>> ChildEdges;
+					ChildEdges.Reserve(Element->Children.Num());
+					for (int32 Child : Element->Children)
+					{
+						ChildEdges.Add(MakeTuple(CurrentEdge.Get<1>(), Child));
+					}
+					Queue.Append(MoveTemp(ChildEdges));
+				}
+			}
+		}
+	}
+
+	void FindLeaves(int32 InInitialVertex, const FDirectedTree& InTree, TArray<int32>& OutLeafVertices)
+	{
+		// Non recursive depth first traversal.
+
+		TArray<int32> Stack({ InInitialVertex });
+		TSet<int32> Visited;
+
+		while (Stack.Num() > 0)
+		{
+			int32 CurrentVertex = Stack.Pop();
+			if (Visited.Contains(CurrentVertex))
+			{
+				continue;
+			}
+
+			Visited.Add(CurrentVertex);
+			const FDirectedTreeElement* Element = InTree.Find(CurrentVertex);
+			// A node with no children (or an absent node) should be considered a leaf : 
+			if ((nullptr != Element) && !Element->Children.IsEmpty())
+			{
+				Stack.Append(Element->Children);
+			}
+			else
+			{
+				OutLeafVertices.Add(CurrentVertex);
+			}
+		}
+	}
+
+	bool DepthFirstTopologicalSort(TArrayView<const int32> InUniqueVertices, TArrayView<const FDirectedEdge> InUniqueEdges, TArray<int32>& OutVertexOrder)
+	{
+		using namespace Private;
+		
+		// Initialize side info about which vertices have been visited and which
+		// have not.
+		TMap<int32, FDepthFirstSortInfo> SortInfos;
+
+		const int32 MaxVertexNum = InUniqueVertices.Num();
+		SortInfos.Reserve(SortInfos.Num() + MaxVertexNum);
+		OutVertexOrder.Reserve(OutVertexOrder.Num() + MaxVertexNum);
+
+		for (int32 Vertex : InUniqueVertices)
+		{
+			SortInfos.Add(Vertex);
+		}
+
+		for (const FDirectedEdge& Edge : InUniqueEdges)
+		{
+			FDepthFirstSortInfo& Info = SortInfos[Edge.Get<0>()];
+			Info.Children.Add(Edge.Get<1>());
+		}
+
+		bool bSuccess = true;
+
+		// Loop through each node to make sure they are all visited.
+		for (int32 Vertex : InUniqueVertices)
+		{
+			FDepthFirstSortInfo& OuterInfo = SortInfos[Vertex];
+
+			if (!(OuterInfo.TemporaryMark || OuterInfo.PermanentMark))
+			{
+				bSuccess = DepthFirstTopologicalSortVisit(Vertex, SortInfos, OutVertexOrder);
+
+				if (!bSuccess)
+				{
+					break;
+				}
+			}
+		}
+
+		// Depth first sort generates it in reverse order. It is reversed here
+		// so it is in forward order. 
+		Algo::Reverse(OutVertexOrder);
+
+		return bSuccess;
+	}
+
+	bool KahnTopologicalSort(TArrayView<const int32> InUniqueVertices, TArrayView<const FDirectedEdge> InUniqueEdges, TArray<int32>& OutVertexOrder)
+	{
+		using FVertexMultiMap = TMultiMap<int32, int32>;
+		using FVertexPair = TMultiMap<int32, int32>::ElementType;
+		
+		FVertexMultiMap Dependencies;
+
+		for (const FDirectedEdge& Edge : InUniqueEdges)
+		{
+			// Destination vertices are dependent upon source vertices.
+			// Create a multimap where each destination vertex references
+			// all of the source vertices that have not been placed in the
+			// output. 
+			Dependencies.Add(Edge.Get<1>(), Edge.Get<0>());
+		}
+
+		TArray<int32> UniqueVertices(InUniqueVertices.GetData(), InUniqueVertices.Num());
+		TArray<int32> IndependentVertices;
+
+		// Sort graph so that vertices with no dependencies always go first.
+		while (UniqueVertices.Num() > 0)
+		{
+			IndependentVertices = UniqueVertices.FilterByPredicate([&](int32 InVertex) { 
+					return (Dependencies.Num(InVertex) == 0);
+			});
+
+			if (0 == IndependentVertices.Num())
+			{
+				// If there are no independent vertices, then likely there is a 
+				// cycle in the graph.
+				return false;
+			}
+
+			FVertexMultiMap UpdatedDependencies;
+			
+			// Remove independent vertices from dependency map.
+			for (FVertexPair& Element : Dependencies)
+			{
+				if (!IndependentVertices.Contains(Element.Value))
+				{
+					UpdatedDependencies.AddUnique(Element.Key, Element.Value);
+				}
+			}
+			Dependencies = UpdatedDependencies;
+
+			// Remove independent vertices from node list
+			for (int32 Vertex : IndependentVertices)
+			{
+				UniqueVertices.RemoveSwap(Vertex);
+			}
+
+			// Add independent vertices to output.
+			OutVertexOrder.Append(IndependentVertices);
+		}
+
+		return true;
+	}
+
+	bool TarjanStronglyConnectedComponents(const TSet<FDirectedEdge>& InEdges, TArray<FStronglyConnectedComponent>& OutComponents, bool bExcludeSingleVertex)
+	{
+		using namespace Private;
+
+		// Use private implementation of Tarjan algorithm.
+		FTarjanAlgoImplSettings Settings;
+		Settings.bExcludeSingleVertex = bExcludeSingleVertex;
+
+		FTarjanAlgoImpl TarjanAlgo(Settings);
+
+		return TarjanAlgo.FindStronglyConnectedComponents(InEdges, OutComponents);
+	}
+} // namespace UE::MathCore::Graph
diff --git a/Engine/Source/Runtime/MathCore/Public/Graph/DirectedGraphUtils.h b/Engine/Source/Runtime/MathCore/Public/Graph/DirectedGraphUtils.h
new file mode 100644
--- /dev/null
+++ b/Engine/Source/Runtime/MathCore/Public/Graph/DirectedGraphUtils.h
@@ -0,0 +1,141 @@
+// Copyright Epic Games, Inc. All Rights Reserved.
+
+#pragma once
+
+#include "CoreMinimal.h"
+
+namespace UE::MathCore::Graph
+{
+	/** A pair of int32s represent a directed edge. The first value represents 
+		* the source vertex, and the second value represents the destination 
+		* vertex.
+		*/
+	using FDirectedEdge = TTuple<int32, int32>;
+
+	/** A strongly connected component contains a subgraph of strongly connected
+		* vertices and their corresponding edges.
+		*/
+	struct FStronglyConnectedComponent
+	{
+		/** Vertices in the strongly connected component. */
+		TArray<int32> Vertices;
+
+		/** Edges in the strongly connected component. */
+		TArray<FDirectedEdge> Edges;
+	};
+
+	/** An element in a directed tree with references to children of a vertex. */
+	struct FDirectedTreeElement
+	{
+		TArray<int32> Children;
+	};
+
+	/** A directed tree graph representation. */
+	using FDirectedTree = TMap<int32, FDirectedTreeElement>;
+
+
+	/** Build a directed tree from an array of edges.
+		*
+		* @parma InEdges - An array of directed edges.
+		* @param OutTree - A tree structure built from the edges.
+		*/
+	MATHCORE_API void BuildDirectedTree(TArrayView<const FDirectedEdge> InEdges, FDirectedTree& OutTree);
+
+	/** Build a transpose directed tree from an array of edges. 
+		*
+		* The transpose of a directed graph is created by reversing each edge.
+		*
+		* @parma InEdges - An array of directed edges.
+		* @param OutTree - A tree structure built from the reversed edges.
+		*/
+	MATHCORE_API void BuildTransposeDirectedTree(TArrayView<const FDirectedEdge> InEdges, FDirectedTree& OutTree);
+
+	/** Traverse a tree in a depth first ordering. Does not traverse the same node twice
+		*
+		* @param InInitialVertex - The starting vertex for traversal.
+		* @param InTree          - The tree structure to traverse. 
+		* @param InVisitFunc     - A function that is called for each vertex that 
+		*                          is visited. If this function returns true, 
+		*                          traversal is continued in a depth first 
+		*                          manner. If this function returns false, the 
+		*                          children of the current vertex are not visited.
+		*/
+	MATHCORE_API void DepthFirstNodeTraversal(int32 InInitialVertex, const FDirectedTree& InTree, TFunctionRef<bool(int32)> InVisitFunc);
+
+	/** Traverse a tree in a depth first ordering. Does not traverse the same edge twice
+		*
+		* @param InInitialVertex - The starting vertex for traversal.
+		* @param InTree          - The tree structure to traverse.
+		* @param InVisitFunc     - A function that is called for each edge that
+		*                          is visited. If this function returns true,
+		*                          the children of this edge's destination vertex will be visited
+		*                          in the future. If this function returns false,
+		*                          then the children will not be visited (unless reaching the node via another route).
+		*/
+	MATHCORE_API void DepthFirstEdgeTraversal(int32 InInitialVertex, const FDirectedTree& InTree, TFunctionRef<bool(int32, int32)> InVisitFunc);
+
+	/** Traverse a tree in a breadth first ordering. Does not traverse the same node twice
+		*
+		* @param InInitialVertex - The starting vertex for traversal.
+		* @param InTree          - The tree structure to traverse. 
+		* @param InVisitFunc     - A function that is called for each vertex that 
+		*                          is visited. If this function returns true, 
+		*                          the children of this vertex will be visited 
+		*                          in the future. If this function returns false, 
+		*                          then the children of the current vertex will 
+		*                          not be visited.
+		*/
+	MATHCORE_API void BreadthFirstNodeTraversal(int32 InInitialVertex, const FDirectedTree& InTree, TFunctionRef<bool(int32)> InVisitFunc);
+
+	/** Traverse a tree in a breadth first ordering. Does not traverse the same edge twice
+		*
+		* @param InInitialVertex - The starting vertex for traversal.
+		* @param InTree          - The tree structure to traverse.
+		* @param InVisitFunc     - A function that is called for each edge that
+		*                          is visited. If this function returns true,
+		*                          the children of this edge's destination vertex will be visited
+		*                          in the future. If this function returns false,
+		*                          then the children will not be visited (unless reaching the node via another route).
+		*/
+	MATHCORE_API void BreadthFirstEdgeTraversal(int32 InInitialVertex, const FDirectedTree& InTree, TFunctionRef<bool(int32, int32)> InVisitFunc);
+
+	/** Returns all leaves (vertex with no children) connected with the starting vertex
+		*
+		* @param InInitialVertex - The starting vertex for traversal.
+		* @param InTree          - The tree structure to traverse.
+		* @param OutLeafVertices - List of all leaf vertices connected to the starting vertex
+		*/
+	MATHCORE_API void FindLeaves(int32 InInitialVertex, const FDirectedTree& InTree, TArray<int32>& OutLeafVertices);
+
+	/** Sort vertices topologically using a depth first sorting algorithm.
+		*
+		* @param InUniqueVertices - An array of vertices to sort.
+		* @param InUniqueEdges - An array of edges describing dependencies.
+		* @param OutVertexOrder - An array where ordered vertices are placed.
+		*
+		* @return True if sorting was successful. False otherwise.
+		*/
+	MATHCORE_API bool DepthFirstTopologicalSort(TArrayView<const int32> InUniqueVertices, TArrayView<const FDirectedEdge> InUniqueEdges, TArray<int32>& OutVertexOrder);
+
+	/** Sort vertices topologically using a Kahn's sorting algorithm.
+		*
+		* @param InUniqueVertices - An array of vertices to sort.
+		* @param InUniqueEdges - An array of edges describing dependencies.
+		* @param OutVertexOrder - An array where ordered vertices are placed.
+		*
+		* @return True if sorting was successful. False otherwise.
+		*/
+	MATHCORE_API bool KahnTopologicalSort(TArrayView<const int32> InUniqueVertices, TArrayView<const FDirectedEdge> InUniqueEdges, TArray<int32>& OutVertexOrder);
+
+
+	/** Find strongly connected components given a set of edges using Tarjan
+		* algorithm.
+		*
+		* @param InEdges - Edges in the directed graph.
+		* @param OutComponents - Strongly connected components found in the graph are be added to this array.
+		* @param bExcludeSingleVertex - If true, single vertices are not be returned as strongly connected components. If false, single vertices may be returned as strongly connected components. 
+		*
+		* @return True if one or more strongly connected components are added to OutComponents. False otherwise. 
+		*/
+	MATHCORE_API bool TarjanStronglyConnectedComponents(const TSet<FDirectedEdge>& InEdges, TArray<FStronglyConnectedComponent>& OutComponents, bool bExcludeSingleVertex = true);
+} // namespace UE::MathCore::Graph
diff --git a/Engine/Source/Runtime/MathCore/README.md b/Engine/Source/Runtime/MathCore/README.md
new file mode 100644
--- /dev/null
+++ b/Engine/Source/Runtime/MathCore/README.md
@@ -0,0 +1,26 @@
+
+# Math Core
+
+**MathCore** is a module meant to regroup any math-related functionality that are not "core-enough" that they should be linked against by any module having a dependency on Core (to avoid code bloat). 
+Unlike Core, which is maintained by the Foundation team, maintenance and support of MathCore is shared between all Epic developers using its content. 
+
+## When to develop in MathCore
+
+* If a piece of code is generic enough that it has no dependency on anything else than Core but would serve a purpose for different modules in the engine. 
+* If the Foundation team is not willing to make this code part of Core.
+* If a piece of code is potentially a candidate for existing in Core but not likely to be used widely enough yet that it justifies it. In time, with the Foundation team's assent, it's possible that such functions would move from MathCore to Core if their usage justifies it (again, it's a question of maintenance).
+
+## When not to develop in MathCore
+
+* If a piece of code is not math-related. 
+* If a function or piece of code would manifestly benefit for any user of the Core module (with the Foundation team's assent, of course)
+
+## When to/not to move code in MathCore
+
+* When the above conditions are all met.
+* If a generic piece of math code is duplicated in several modules already, it's a probably a good candidate to be moved in MathCore, as long as the original locations are updated, in order not to make this yet another duplicate of the same piece of functionality. 
+* If the duplicated code differs too much and its complexity would increase too greatly if made generic, then it's probably worth keeping things separate and let each module have the ability to customize its content freely for its own specific purposes. 
+
+## Tests
+
+It is strongly encouraged to write Low-Level Tests when adding/modifying anything in MathCore. There's already a dedicated executable, MathCoreTests, for that purpose. It is compiling and running on all platforms and is run on CIS so it's the perfect place to extend the test coverage of MathCore code. 
diff --git a/Engine/Source/Runtime/SignalProcessing/Private/DirectedGraphAlgo.cpp b/Engine/Source/Runtime/SignalProcessing/Private/DirectedGraphAlgo.cpp
--- a/Engine/Source/Runtime/SignalProcessing/Private/DirectedGraphAlgo.cpp
+++ b/Engine/Source/Runtime/SignalProcessing/Private/DirectedGraphAlgo.cpp
@@ -2,8 +2,12 @@
 
 #include "DSP/DirectedGraphAlgo.h"
 
+// This entire file has been deprecated/moved to MathCore (Graph/DirectedGraphUtils.h)
+PRAGMA_DISABLE_DEPRECATION_WARNINGS
+
 namespace Audio
 {
+
 	namespace DirectedGraphAlgoPrivate
 	{
 		// Used to track which vertices have been visited and vertex children.
@@ -451,4 +455,7 @@ namespace Audio
 
 		return TarjanAlgo.FindStronglyConnectedComponents(InEdges, OutComponents);
 	}
+
 }
+
+PRAGMA_ENABLE_DEPRECATION_WARNINGS
\ No newline at end of file
diff --git a/Engine/Source/Runtime/SignalProcessing/Public/DSP/DirectedGraphAlgo.h b/Engine/Source/Runtime/SignalProcessing/Public/DSP/DirectedGraphAlgo.h
--- a/Engine/Source/Runtime/SignalProcessing/Public/DSP/DirectedGraphAlgo.h
+++ b/Engine/Source/Runtime/SignalProcessing/Public/DSP/DirectedGraphAlgo.h
@@ -4,35 +4,41 @@
 
 #include "CoreMinimal.h"
 
+#include "Graph/DirectedGraphUtils.h"
+
 namespace Audio
 {
+
 	/** A pair of int32s represent a directed edge. The first value represents 
 	 * the source vertex, and the second value represents the destination 
 	 * vertex.
 	 */
+	UE_DEPRECATED(5.5, "This has moved to the UE::MathCore::Graph namespace")
 	typedef TTuple<int32, int32> FDirectedEdge;
 
 	/** A strongly connected component contains a subgraph of strongly connected
 	 * vertices and their corresponding edges.
 	 */
-	struct FStronglyConnectedComponent
+	struct UE_DEPRECATED(5.5, "This has moved to the UE::MathCore::Graph namespace") FStronglyConnectedComponent
 	{
 		/** Vertices in the strongly connected component. */
 		TArray<int32> Vertices;
 
 		/** Edges in the strongly connected component. */
+		PRAGMA_DISABLE_DEPRECATION_WARNINGS
 		TArray<FDirectedEdge> Edges;
+		PRAGMA_ENABLE_DEPRECATION_WARNINGS
 	};
 
 	/** An element in a directed tree with references to children of a vertex. */
-	struct FDirectedTreeElement
-	{
-		TArray<int32> Children;
-	};
+	UE_DEPRECATED(5.5, "This has moved to the UE::MathCore::Graph namespace")
+	typedef UE::MathCore::Graph::FDirectedTreeElement FDirectedTreeElement;
 
 	/** A directed tree graph represenation. */
+	UE_DEPRECATED(5.5, "This has moved to the UE::MathCore::Graph namespace")
+	PRAGMA_DISABLE_DEPRECATION_WARNINGS
 	typedef TMap<int32, FDirectedTreeElement> FDirectedTree;
-
+	PRAGMA_ENABLE_DEPRECATION_WARNINGS
 
 	struct FDirectedGraphAlgo
 	{
@@ -41,7 +47,10 @@ namespace Audio
 		 * @parma InEdges - An array of directed eges.
 		 * @param OutTree - A tree structure built from the edges.
 		 */
+		UE_DEPRECATED(5.5, "This has moved to the UE::MathCore::Graph namespace")
+		PRAGMA_DISABLE_DEPRECATION_WARNINGS
 		static SIGNALPROCESSING_API void BuildDirectedTree(TArrayView<const FDirectedEdge> InEdges, FDirectedTree& OutTree);
+		PRAGMA_ENABLE_DEPRECATION_WARNINGS
 
 		/** Build a transpose directed tree from an array of edges. 
 		 *
@@ -50,7 +59,10 @@ namespace Audio
 		 * @parma InEdges - An array of directed eges.
 		 * @param OutTree - A tree structure built from the reversed edges.
 		 */
+		UE_DEPRECATED(5.5, "This has moved to the UE::MathCore::Graph namespace")
+		PRAGMA_DISABLE_DEPRECATION_WARNINGS
 		static SIGNALPROCESSING_API void BuildTransposeDirectedTree(TArrayView<const FDirectedEdge> InEdges, FDirectedTree& OutTree);
+		PRAGMA_ENABLE_DEPRECATION_WARNINGS
 
 		/** Traverse a tree in a depth first ordering.
 		 *
@@ -62,7 +74,10 @@ namespace Audio
 		 *                          manner. If this function returns false, the 
 		 *                          children of the current vertex are not visited.
 		 */
+		UE_DEPRECATED(5.5, "This has moved to the UE::MathCore::Graph namespace")
+		PRAGMA_DISABLE_DEPRECATION_WARNINGS
 		static SIGNALPROCESSING_API void DepthFirstTraversal(int32 InInitialVertex, const FDirectedTree& InTree, TFunctionRef<bool (int32)> InVisitFunc);
+		PRAGMA_ENABLE_DEPRECATION_WARNINGS
 
 		/** Traverse a tree in a breadth first ordering.
 		 *
@@ -75,7 +90,10 @@ namespace Audio
 		 *                          then the children of the current vertex will 
 		 *                          not be visited.
 		 */
+		UE_DEPRECATED(5.5, "This has moved to the UE::MathCore::Graph namespace")
+		PRAGMA_DISABLE_DEPRECATION_WARNINGS
 		static SIGNALPROCESSING_API void BreadthFirstTraversal(int32 InInitialVertex, const FDirectedTree& InTree, TFunctionRef<bool (int32)> InVisitFunc);
+		PRAGMA_ENABLE_DEPRECATION_WARNINGS
 
 		/** Sort vertices topologically using a depth first sorting algorithm.
 		 *
@@ -85,7 +103,10 @@ namespace Audio
 		 *
 		 * @return True if sorting was successful. False otherwise.
 		 */
+		UE_DEPRECATED(5.5, "This has moved to the UE::MathCore::Graph namespace")
+		PRAGMA_DISABLE_DEPRECATION_WARNINGS
 		static SIGNALPROCESSING_API bool DepthFirstTopologicalSort(TArrayView<const int32> InUniqueVertices, TArrayView<const FDirectedEdge> InUniqueEdges, TArray<int32>& OutVertexOrder);
+		PRAGMA_ENABLE_DEPRECATION_WARNINGS
 
 		/** Sort vertices topologically using a Kahn's sorting algorithm.
 		 *
@@ -95,7 +116,10 @@ namespace Audio
 		 *
 		 * @return True if sorting was successful. False otherwise.
 		 */
+		UE_DEPRECATED(5.5, "This has moved to the UE::MathCore::Graph namespace")
+		PRAGMA_DISABLE_DEPRECATION_WARNINGS
 		static SIGNALPROCESSING_API bool KahnTopologicalSort(TArrayView<const int32> InUniqueVertices, TArrayView<const FDirectedEdge> InUniqueEdges, TArray<int32>& OutVertexOrder);
+		PRAGMA_ENABLE_DEPRECATION_WARNINGS
 
 
 		/** Find strongly connected components given a set of edges using Tarjan
@@ -107,6 +131,9 @@ namespace Audio
 		 *
 		 * @return True if one or more strongly connected components are added to OutComponents. False otherwise. 
 		 */
+		UE_DEPRECATED(5.5, "This has moved to the UE::MathCore::Graph namespace")
+		PRAGMA_DISABLE_DEPRECATION_WARNINGS
 		static SIGNALPROCESSING_API bool TarjanStronglyConnectedComponents(const TSet<FDirectedEdge>& InEdges, TArray<FStronglyConnectedComponent>& OutComponents, bool bExcludeSingleVertex = true);
+		PRAGMA_ENABLE_DEPRECATION_WARNINGS
 	};
 }
diff --git a/Engine/Source/Runtime/SignalProcessing/SignalProcessing.Build.cs b/Engine/Source/Runtime/SignalProcessing/SignalProcessing.Build.cs
--- a/Engine/Source/Runtime/SignalProcessing/SignalProcessing.Build.cs
+++ b/Engine/Source/Runtime/SignalProcessing/SignalProcessing.Build.cs
@@ -9,7 +9,8 @@ namespace UnrealBuildTool.Rules
 			PrivateDependencyModuleNames.AddRange(
 				new string[]
 				{
-					"Core"
+					"Core",
+					"MathCore",
 				}
 			);
 
