Correctly set compilation flags for FoundationTests to build with UBT without -Mode=Test

diff --git a/Engine/Source/Programs/LowLevelTests/FoundationTests/FoundationTests.Target.cs b/Engine/Source/Programs/LowLevelTests/FoundationTests/FoundationTests.Target.cs
--- a/Engine/Source/Programs/LowLevelTests/FoundationTests/FoundationTests.Target.cs
+++ b/Engine/Source/Programs/LowLevelTests/FoundationTests/FoundationTests.Target.cs
@@ -11,6 +11,8 @@ public class FoundationTestsTarget : TestTargetRules
 		bWithLowLevelTestsOverride = true;
 		bCompileWithPluginSupport = true;
 
+		bCompileAgainstCoreUObject = true;
+
 		bBuildWithEditorOnlyData = Target.Platform.IsInGroup(UnrealPlatformGroup.Desktop)
 			&& (Target.Configuration == UnrealTargetConfiguration.Debug || Target.Configuration == UnrealTargetConfiguration.Development);
 	}
