Rrestrict test metadata to test mode
Step verified job 66390cb0242b321d4c9d98a6

diff --git a/Engine/Source/Developer/LowLevelTestsRunner/README.md b/Engine/Source/Developer/LowLevelTestsRunner/README.md
--- a/Engine/Source/Developer/LowLevelTestsRunner/README.md
+++ b/Engine/Source/Developer/LowLevelTestsRunner/README.md
@@ -94,12 +94,15 @@ Follow these steps to create your explicit test:
 	```csharp
 	public class UEModuleTests : TestModuleRules
 	{
-        sttic UEModuleTests()
+        static UEModuleTests()
         {
-          TestMetadata = new Metadata();
-		  TestMetadata.TestName = "UEModule";
-		  TestMetadata.TestShortName = "UE Module";
-		  // any other additional properties
+          if (InTestMode)
+          {
+            TestMetadata = new Metadata();
+		    TestMetadata.TestName = "UEModule";
+		    TestMetadata.TestShortName = "UE Module";
+		    // any other additional properties
+          }
         }
 
         public UEModuleTests(ReadOnlyTargetRules Target) : base(Target)
diff --git a/Engine/Source/Programs/ChaosUserDataPTTests/ChaosUserDataPTTests.Build.cs b/Engine/Source/Programs/ChaosUserDataPTTests/ChaosUserDataPTTests.Build.cs
--- a/Engine/Source/Programs/ChaosUserDataPTTests/ChaosUserDataPTTests.Build.cs
+++ b/Engine/Source/Programs/ChaosUserDataPTTests/ChaosUserDataPTTests.Build.cs
@@ -6,9 +6,12 @@ public class ChaosUserDataPTTests : TestModuleRules
 {
 	static ChaosUserDataPTTests()
 	{
-		TestMetadata = new Metadata();
-		TestMetadata.TestName = "ChasoUserDataPT";
-		TestMetadata.TestShortName = "Chaos User Data PT";
+		if (InTestMode)
+		{
+			TestMetadata = new Metadata();
+			TestMetadata.TestName = "ChasoUserDataPT";
+			TestMetadata.TestShortName = "Chaos User Data PT";
+		}
 	}
 
 	public ChaosUserDataPTTests(ReadOnlyTargetRules Target) : base(Target)
diff --git a/Engine/Source/Programs/HeadlessChaos/HeadlessChaos.Build.cs b/Engine/Source/Programs/HeadlessChaos/HeadlessChaos.Build.cs
--- a/Engine/Source/Programs/HeadlessChaos/HeadlessChaos.Build.cs
+++ b/Engine/Source/Programs/HeadlessChaos/HeadlessChaos.Build.cs
@@ -6,10 +6,13 @@ public class HeadlessChaos : TestModuleRules
 {
 	static HeadlessChaos()
 	{
-		TestMetadata = new Metadata();
-		TestMetadata.TestName = "HeadlessChaos";
-		TestMetadata.TestShortName = "Headless Chaos";
-		TestMetadata.UsesCatch2 = false;
+		if (InTestMode)
+		{
+			TestMetadata = new Metadata();
+			TestMetadata.TestName = "HeadlessChaos";
+			TestMetadata.TestShortName = "Headless Chaos";
+			TestMetadata.UsesCatch2 = false;
+		}
 	}
 
 	public HeadlessChaos(ReadOnlyTargetRules Target) : base(Target, false)
diff --git a/Engine/Source/Programs/LowLevelTests/FoundationTests/FoundationTests.Build.cs b/Engine/Source/Programs/LowLevelTests/FoundationTests/FoundationTests.Build.cs
--- a/Engine/Source/Programs/LowLevelTests/FoundationTests/FoundationTests.Build.cs
+++ b/Engine/Source/Programs/LowLevelTests/FoundationTests/FoundationTests.Build.cs
@@ -5,36 +5,39 @@ public class FoundationTests : TestModuleRules
 {
 	static FoundationTests()
 	{
-		TestMetadata = new Metadata();
-		TestMetadata.TestName = "Foundation";
-		TestMetadata.TestShortName = "Foundation";
-		TestMetadata.ReportType = "xml";
-		TestMetadata.SupportedPlatforms.Add(UnrealTargetPlatform.Linux);
-		TestMetadata.SupportedPlatforms.Add(UnrealTargetPlatform.Mac);
-		TestMetadata.SupportedPlatforms.Add(UnrealTargetPlatform.Android);
-		TestMetadata.SupportedPlatforms.Add(UnrealTargetPlatform.IOS);
-		
-		string PlatformCompilationArgs;
-		foreach (var Platform in UnrealTargetPlatform.GetValidPlatforms())
+		if (InTestMode)
 		{
-			if (Platform == UnrealTargetPlatform.Android)
-			{
-				PlatformCompilationArgs = "-allmodules -architectures=arm64";
-			}
-			else
+			TestMetadata = new Metadata();
+			TestMetadata.TestName = "Foundation";
+			TestMetadata.TestShortName = "Foundation";
+			TestMetadata.ReportType = "xml";
+			TestMetadata.SupportedPlatforms.Add(UnrealTargetPlatform.Linux);
+			TestMetadata.SupportedPlatforms.Add(UnrealTargetPlatform.Mac);
+			TestMetadata.SupportedPlatforms.Add(UnrealTargetPlatform.Android);
+			TestMetadata.SupportedPlatforms.Add(UnrealTargetPlatform.IOS);
+
+			string PlatformCompilationArgs;
+			foreach (var Platform in UnrealTargetPlatform.GetValidPlatforms())
 			{
-				PlatformCompilationArgs = "-allmodules";
+				if (Platform == UnrealTargetPlatform.Android)
+				{
+					PlatformCompilationArgs = "-allmodules -architectures=arm64";
+				}
+				else
+				{
+					PlatformCompilationArgs = "-allmodules";
+				}
+				TestMetadata.PlatformCompilationExtraArgs.Add(Platform, PlatformCompilationArgs);
 			}
-			TestMetadata.PlatformCompilationExtraArgs.Add(Platform, PlatformCompilationArgs);
-		}
 
-		// Platform-specific tags
-		TestMetadata.PlatformTags.Add(UnrealTargetPlatform.Linux, "~[.]~[Slow]");
-		TestMetadata.PlatformTags.Add(UnrealTargetPlatform.Android, "~[Perf]~[Slow]~[AndroidSkip]");
+			// Platform-specific tags
+			TestMetadata.PlatformTags.Add(UnrealTargetPlatform.Linux, "~[.]~[Slow]");
+			TestMetadata.PlatformTags.Add(UnrealTargetPlatform.Android, "~[Perf]~[Slow]~[AndroidSkip]");
 
-		// Allow Android run for this test
-		// Will remove Android from PlatformsRunUnsupported as more diverse types of tests can run on this platform
-		TestMetadata.PlatformsRunUnsupported.Remove(UnrealTargetPlatform.Android);
+			// Allow Android run for this test
+			// Will remove Android from PlatformsRunUnsupported as more diverse types of tests can run on this platform
+			TestMetadata.PlatformsRunUnsupported.Remove(UnrealTargetPlatform.Android);
+		}
 	}
 	public FoundationTests(ReadOnlyTargetRules Target) : base(Target, true)
 	{
diff --git a/Engine/Source/Programs/LowLevelTests/MathCoreTests/MathCoreTests.Build.cs b/Engine/Source/Programs/LowLevelTests/MathCoreTests/MathCoreTests.Build.cs
--- a/Engine/Source/Programs/LowLevelTests/MathCoreTests/MathCoreTests.Build.cs
+++ b/Engine/Source/Programs/LowLevelTests/MathCoreTests/MathCoreTests.Build.cs
@@ -5,29 +5,32 @@ public class MathCoreTests : TestModuleRules
 {
 	static MathCoreTests()
 	{
-		TestMetadata = new Metadata();
-		TestMetadata.TestName = "MathCore";
-		TestMetadata.TestShortName = "MathCore";
-		TestMetadata.ReportType = "xml";
-		TestMetadata.SupportedPlatforms.Add(UnrealTargetPlatform.Linux);
-		TestMetadata.SupportedPlatforms.Add(UnrealTargetPlatform.Mac);
-		TestMetadata.SupportedPlatforms.Add(UnrealTargetPlatform.Android);
-		TestMetadata.SupportedPlatforms.Add(UnrealTargetPlatform.IOS);
-
-		string PlatformCompilationArgs;
-		foreach (var Platform in UnrealTargetPlatform.GetValidPlatforms())
+		if (InTestMode)
 		{
-			if (Platform == UnrealTargetPlatform.Android)
-			{
-				PlatformCompilationArgs = "-allmodules -architectures=arm64";
-			}
-			else
+			TestMetadata = new Metadata();
+			TestMetadata.TestName = "MathCore";
+			TestMetadata.TestShortName = "MathCore";
+			TestMetadata.ReportType = "xml";
+			TestMetadata.SupportedPlatforms.Add(UnrealTargetPlatform.Linux);
+			TestMetadata.SupportedPlatforms.Add(UnrealTargetPlatform.Mac);
+			TestMetadata.SupportedPlatforms.Add(UnrealTargetPlatform.Android);
+			TestMetadata.SupportedPlatforms.Add(UnrealTargetPlatform.IOS);
+
+			string PlatformCompilationArgs;
+			foreach (var Platform in UnrealTargetPlatform.GetValidPlatforms())
 			{
-				PlatformCompilationArgs = "-allmodules";
+				if (Platform == UnrealTargetPlatform.Android)
+				{
+					PlatformCompilationArgs = "-allmodules -architectures=arm64";
+				}
+				else
+				{
+					PlatformCompilationArgs = "-allmodules";
+				}
+				TestMetadata.PlatformCompilationExtraArgs.Add(Platform, PlatformCompilationArgs);
 			}
-			TestMetadata.PlatformCompilationExtraArgs.Add(Platform, PlatformCompilationArgs);
+			TestMetadata.PlatformsRunUnsupported.Remove(UnrealTargetPlatform.Android);
 		}
-		TestMetadata.PlatformsRunUnsupported.Remove(UnrealTargetPlatform.Android);
 	}
 
 	public MathCoreTests(ReadOnlyTargetRules Target) : base(Target, InUsesCatch2:true)
diff --git a/Engine/Source/Programs/NetworkPredictionTests/NetworkPredictionTests.Build.cs b/Engine/Source/Programs/NetworkPredictionTests/NetworkPredictionTests.Build.cs
--- a/Engine/Source/Programs/NetworkPredictionTests/NetworkPredictionTests.Build.cs
+++ b/Engine/Source/Programs/NetworkPredictionTests/NetworkPredictionTests.Build.cs
@@ -7,10 +7,13 @@ public class NetworkPredictionTests : TestModuleRules
 {
 	static NetworkPredictionTests()
 	{
-		TestMetadata = new Metadata();
-		TestMetadata.TestName = "NetworkPredictionPlugin";
-		TestMetadata.TestShortName = "Net Prediction";
-		TestMetadata.SupportedPlatforms.Add(UnrealTargetPlatform.Linux);
+		if (InTestMode)
+		{
+			TestMetadata = new Metadata();
+			TestMetadata.TestName = "NetworkPredictionPlugin";
+			TestMetadata.TestShortName = "Net Prediction";
+			TestMetadata.SupportedPlatforms.Add(UnrealTargetPlatform.Linux);
+		}
 	}
 
 	public NetworkPredictionTests(ReadOnlyTargetRules Target) : base(Target)
diff --git a/Engine/Source/Programs/PlainPropsTests/PlainPropsTests.Build.cs b/Engine/Source/Programs/PlainPropsTests/PlainPropsTests.Build.cs
--- a/Engine/Source/Programs/PlainPropsTests/PlainPropsTests.Build.cs
+++ b/Engine/Source/Programs/PlainPropsTests/PlainPropsTests.Build.cs
@@ -5,22 +5,25 @@ public class PlainPropsTests : TestModuleRules
 {
 	static PlainPropsTests()
 	{
-		TestMetadata = new Metadata();
-		TestMetadata.TestName = "PlainProps";
-		TestMetadata.TestShortName = "PlainProps";
-		TestMetadata.ReportType = "xml";
-		TestMetadata.SupportedPlatforms.Add(UnrealTargetPlatform.Linux);
-		TestMetadata.SupportedPlatforms.Add(UnrealTargetPlatform.Mac);
-		
-		string PlatformCompilationArgs;
-		foreach (var Platform in UnrealTargetPlatform.GetValidPlatforms())
+		if (InTestMode)
 		{
-			PlatformCompilationArgs = "-allmodules";
-			TestMetadata.PlatformCompilationExtraArgs.Add(Platform, PlatformCompilationArgs);
-		}
+			TestMetadata = new Metadata();
+			TestMetadata.TestName = "PlainProps";
+			TestMetadata.TestShortName = "PlainProps";
+			TestMetadata.ReportType = "xml";
+			TestMetadata.SupportedPlatforms.Add(UnrealTargetPlatform.Linux);
+			TestMetadata.SupportedPlatforms.Add(UnrealTargetPlatform.Mac);
+
+			string PlatformCompilationArgs;
+			foreach (var Platform in UnrealTargetPlatform.GetValidPlatforms())
+			{
+				PlatformCompilationArgs = "-allmodules";
+				TestMetadata.PlatformCompilationExtraArgs.Add(Platform, PlatformCompilationArgs);
+			}
 
-		// Platform-specific tags
-		TestMetadata.PlatformTags.Add(UnrealTargetPlatform.Linux, "~[.]~[Slow]");
+			// Platform-specific tags
+			TestMetadata.PlatformTags.Add(UnrealTargetPlatform.Linux, "~[.]~[Slow]");
+		}
 	}
 	public PlainPropsTests(ReadOnlyTargetRules Target) : base(Target, true)
 	{
diff --git a/Engine/Source/Programs/ReplicationSystemLowLevelTests/ReplicationSystemLowLevelTests.Build.cs b/Engine/Source/Programs/ReplicationSystemLowLevelTests/ReplicationSystemLowLevelTests.Build.cs
--- a/Engine/Source/Programs/ReplicationSystemLowLevelTests/ReplicationSystemLowLevelTests.Build.cs
+++ b/Engine/Source/Programs/ReplicationSystemLowLevelTests/ReplicationSystemLowLevelTests.Build.cs
@@ -6,12 +6,15 @@ public class ReplicationSystemLowLevelTests : TestModuleRules
 {
 	static ReplicationSystemLowLevelTests()
 	{
-		TestMetadata = new Metadata();
-		TestMetadata.TestName = "ReplicationSystem";
-		TestMetadata.TestShortName = "Replication System";
-		TestMetadata.SupportedPlatforms.Add(UnrealTargetPlatform.Linux);
-		TestMetadata.SupportedPlatforms.Add(UnrealTargetPlatform.LinuxArm64);
-		TestMetadata.PlatformRunContainerized.Add(UnrealTargetPlatform.LinuxArm64, true);
+		if (InTestMode)
+		{
+			TestMetadata = new Metadata();
+			TestMetadata.TestName = "ReplicationSystem";
+			TestMetadata.TestShortName = "Replication System";
+			TestMetadata.SupportedPlatforms.Add(UnrealTargetPlatform.Linux);
+			TestMetadata.SupportedPlatforms.Add(UnrealTargetPlatform.LinuxArm64);
+			TestMetadata.PlatformRunContainerized.Add(UnrealTargetPlatform.LinuxArm64, true);
+		}
 	}
 
 	public ReplicationSystemLowLevelTests(ReadOnlyTargetRules Target) : base(Target)
diff --git a/Engine/Source/Programs/SwitchboardListener/SwitchboardListenerTests/SwitchboardListenerTests.Build.cs b/Engine/Source/Programs/SwitchboardListener/SwitchboardListenerTests/SwitchboardListenerTests.Build.cs
--- a/Engine/Source/Programs/SwitchboardListener/SwitchboardListenerTests/SwitchboardListenerTests.Build.cs
+++ b/Engine/Source/Programs/SwitchboardListener/SwitchboardListenerTests/SwitchboardListenerTests.Build.cs
@@ -7,9 +7,12 @@ public class SwitchboardListenerTests : TestModuleRules
 {
 	static SwitchboardListenerTests()
 	{
-		TestMetadata = new Metadata();
-		TestMetadata.TestName = "SwitchboardListener";
-		TestMetadata.TestShortName = "SwitchboardListener";
+		if (InTestMode)
+		{
+			TestMetadata = new Metadata();
+			TestMetadata.TestName = "SwitchboardListener";
+			TestMetadata.TestShortName = "SwitchboardListener";
+		}
 	}
 
 	public SwitchboardListenerTests(ReadOnlyTargetRules Target) : base(Target)
diff --git a/Engine/Source/Programs/ToolMenusTests/ToolMenusTests.Build.cs b/Engine/Source/Programs/ToolMenusTests/ToolMenusTests.Build.cs
--- a/Engine/Source/Programs/ToolMenusTests/ToolMenusTests.Build.cs
+++ b/Engine/Source/Programs/ToolMenusTests/ToolMenusTests.Build.cs
@@ -5,25 +5,28 @@ public class ToolMenusTests : TestModuleRules
 {
 	static ToolMenusTests()
 	{
-		TestMetadata = new Metadata();
-		TestMetadata.TestName = "ToolMenus";
-		TestMetadata.TestShortName = "ToolMenus";
-		TestMetadata.ReportType = "xml";
-		TestMetadata.SupportedPlatforms.Add(UnrealTargetPlatform.Linux);
-		TestMetadata.SupportedPlatforms.Add(UnrealTargetPlatform.Mac);
-		
-		string PlatformCompilationArgs;
-		foreach (var Platform in UnrealTargetPlatform.GetValidPlatforms())
+		if (InTestMode)
 		{
-			if (Platform == UnrealTargetPlatform.Android)
-			{
-				PlatformCompilationArgs = "-allmodules -architectures=arm64";
-			}
-			else
+			TestMetadata = new Metadata();
+			TestMetadata.TestName = "ToolMenus";
+			TestMetadata.TestShortName = "ToolMenus";
+			TestMetadata.ReportType = "xml";
+			TestMetadata.SupportedPlatforms.Add(UnrealTargetPlatform.Linux);
+			TestMetadata.SupportedPlatforms.Add(UnrealTargetPlatform.Mac);
+
+			string PlatformCompilationArgs;
+			foreach (var Platform in UnrealTargetPlatform.GetValidPlatforms())
 			{
-				PlatformCompilationArgs = "-allmodules";
+				if (Platform == UnrealTargetPlatform.Android)
+				{
+					PlatformCompilationArgs = "-allmodules -architectures=arm64";
+				}
+				else
+				{
+					PlatformCompilationArgs = "-allmodules";
+				}
+				TestMetadata.PlatformCompilationExtraArgs.Add(Platform, PlatformCompilationArgs);
 			}
-			TestMetadata.PlatformCompilationExtraArgs.Add(Platform, PlatformCompilationArgs);
 		}
 	}
 
diff --git a/Engine/Source/Programs/UnrealBuildTool/Configuration/TestModuleRules.cs b/Engine/Source/Programs/UnrealBuildTool/Configuration/TestModuleRules.cs
--- a/Engine/Source/Programs/UnrealBuildTool/Configuration/TestModuleRules.cs
+++ b/Engine/Source/Programs/UnrealBuildTool/Configuration/TestModuleRules.cs
@@ -28,6 +28,11 @@ namespace UnrealBuildTool
 
 		private bool bUsesCatch2 = true;
 
+		/// <summary>
+		/// Check if running in test mode.
+		/// </summary>
+		protected static bool InTestMode = Environment.GetCommandLineArgs().Contains("-Mode=Test");
+
 		/// <summary>
 		/// Associated tested module of this test module.
 		/// </summary>
diff --git a/Engine/Source/Programs/WebTests/WebTests.Build.cs b/Engine/Source/Programs/WebTests/WebTests.Build.cs
--- a/Engine/Source/Programs/WebTests/WebTests.Build.cs
+++ b/Engine/Source/Programs/WebTests/WebTests.Build.cs
@@ -6,24 +6,27 @@ public class WebTests : TestModuleRules
 {
 	static WebTests()
 	{
-		TestMetadata = new Metadata();
-		TestMetadata.TestName = "Web";
-		TestMetadata.TestShortName = "Web";
-		TestMetadata.ReportType = "xml";
-		TestMetadata.SupportedPlatforms.Add(UnrealTargetPlatform.Linux);
-		TestMetadata.SupportedPlatforms.Add(UnrealTargetPlatform.Mac);
-		TestMetadata.SupportedPlatforms.Add(UnrealTargetPlatform.Android);
-		TestMetadata.SupportedPlatforms.Add(UnrealTargetPlatform.IOS);
+		if (InTestMode)
+		{
+			TestMetadata = new Metadata();
+			TestMetadata.TestName = "Web";
+			TestMetadata.TestShortName = "Web";
+			TestMetadata.ReportType = "xml";
+			TestMetadata.SupportedPlatforms.Add(UnrealTargetPlatform.Linux);
+			TestMetadata.SupportedPlatforms.Add(UnrealTargetPlatform.Mac);
+			TestMetadata.SupportedPlatforms.Add(UnrealTargetPlatform.Android);
+			TestMetadata.SupportedPlatforms.Add(UnrealTargetPlatform.IOS);
 
-		TestMetadata.HasAfterSteps = true;
-		TestMetadata.PlatformCompilationExtraArgs.Add(UnrealTargetPlatform.Win64, "-allmodules");
-		TestMetadata.PlatformCompilationExtraArgs.Add(UnrealTargetPlatform.Linux, "-allmodules");
-		TestMetadata.PlatformCompilationExtraArgs.Add(UnrealTargetPlatform.LinuxArm64, "-allmodules");
-		TestMetadata.PlatformCompilationExtraArgs.Add(UnrealTargetPlatform.Mac, "-allmodules");
-		TestMetadata.PlatformCompilationExtraArgs.Add(UnrealTargetPlatform.IOS, "-allmodules");
-		TestMetadata.PlatformCompilationExtraArgs.Add(UnrealTargetPlatform.TVOS, "-allmodules");
-		TestMetadata.PlatformCompilationExtraArgs.Add(UnrealTargetPlatform.Android, "-allmodules  -architectures=arm64");
-		TestMetadata.PlatformCompilationExtraArgs.Add(UnrealTargetPlatform.VisionOS, "-allmodules");
+			TestMetadata.HasAfterSteps = true;
+			TestMetadata.PlatformCompilationExtraArgs.Add(UnrealTargetPlatform.Win64, "-allmodules");
+			TestMetadata.PlatformCompilationExtraArgs.Add(UnrealTargetPlatform.Linux, "-allmodules");
+			TestMetadata.PlatformCompilationExtraArgs.Add(UnrealTargetPlatform.LinuxArm64, "-allmodules");
+			TestMetadata.PlatformCompilationExtraArgs.Add(UnrealTargetPlatform.Mac, "-allmodules");
+			TestMetadata.PlatformCompilationExtraArgs.Add(UnrealTargetPlatform.IOS, "-allmodules");
+			TestMetadata.PlatformCompilationExtraArgs.Add(UnrealTargetPlatform.TVOS, "-allmodules");
+			TestMetadata.PlatformCompilationExtraArgs.Add(UnrealTargetPlatform.Android, "-allmodules  -architectures=arm64");
+			TestMetadata.PlatformCompilationExtraArgs.Add(UnrealTargetPlatform.VisionOS, "-allmodules");
+		}
 	}
 
 	public WebTests(ReadOnlyTargetRules Target) : base(Target)
