UE-219092 - Warning: Flushing package <path> while it's being preloaded in the same callstack is not possible

Cyclic dependencies between Blueprints that also include a UDS in the mix can create a scenario where we might recursively load a package that we're already loading. Currently, Zen loader warns about this. While technically harmless, it makes the warning somewhat useless if it raises false positives.

The core issue is that preloading a UDS will use FObjectPropertyBase::FindImportedObject to resolve a string that represents a UObject path. In the case of a cycle, RequestPlaceholderValue creates a placeholder object, but it neglects to supply a package index for the import. As a result, ResolveDependencyPlaceholder forces a StaticLoadObject call that triggers the warning. This can be fixed by determining the package index of the corresponding object path in RequestPlaceholderValue, which avoids trying to find the import using StaticLoadObject.

This fix introduces two new API functions:
* FLinkerLoad::FindImport overload that takes a full object path and returns the corresponding package index.
* FPackageName::SplitFullObjectPath overload that returns each subobject in a path.

For this fix, the subobject support isn't necessary. However, FindImport's implementation would be incomplete if it didn't support subobjects. New tests for each function were also added to LowLevelTests. In order to test FLinkerLoad, we also need to compile out the "final" when running the tests. This makes it easier to setup a test harness for strictly testing FindImport functionality.

#rnx

diff --git a/Engine/Source/Programs/LowLevelTests/FoundationTests/Tests/LinkerLoadTests.cpp b/Engine/Source/Programs/LowLevelTests/FoundationTests/Tests/LinkerLoadTests.cpp
new file mode 100644
--- /dev/null
+++ b/Engine/Source/Programs/LowLevelTests/FoundationTests/Tests/LinkerLoadTests.cpp
@@ -0,0 +1,181 @@
+// Copyright Epic Games, Inc. All Rights Reserved.
+
+#include "CoreTypes.h"
+
+#include "UObject/LinkerLoad.h"
+#include "UObject/Package.h"
+#include "UObject/UObjectThreadContext.h"
+#include "TestHarness.h"
+
+struct FTestLinkerLoad : public FLinkerLoad
+{
+	FTestLinkerLoad(UPackage* InParent, const FPackagePath& InPackagePath)
+		: FLinkerLoad(InParent, InPackagePath, 0)
+	{
+		FObjectImport PackageImport;
+		PackageImport.ObjectName = TEXT("/Path/To/A/Package");
+		PackageImport.ClassName = NAME_Package;
+		PackageIndex = ImportMap.Num();
+		ImportMap.Add(PackageImport);
+
+		FObjectImport RootObjectImport;
+		RootObjectImport.ObjectName = TEXT("Object");
+		RootObjectImport.OuterIndex = FPackageIndex::FromImport(PackageIndex);
+		RootObjectIndex = ImportMap.Num();
+		ImportMap.Add(RootObjectImport);
+
+		FObjectImport SubObject1Import;
+		SubObject1Import.ObjectName = TEXT("SubObject1");
+		SubObject1Import.OuterIndex = FPackageIndex::FromImport(RootObjectIndex);
+		SubObject1Index = ImportMap.Num();
+		ImportMap.Add(SubObject1Import);
+
+		FObjectImport SubObject2Import;
+		SubObject2Import.ObjectName = TEXT("SubObject2");
+		SubObject2Import.OuterIndex = FPackageIndex::FromImport(SubObject1Index);
+		SubObject2Index = ImportMap.Num();
+		ImportMap.Add(SubObject2Import);
+	}
+
+	int32 PackageIndex = -1;
+	int32 RootObjectIndex = -1;
+	int32 SubObject1Index = -1;
+	int32 SubObject2Index = -1;
+};
+
+TEST_CASE("Validate that a full object path returns import entries", "[CoreUObject][FLinkerLoad::FindImport]")
+{
+	FPackagePath TestPath = FPackagePath::FromPackageNameChecked("/game/TestPackage");
+	TestPath.SetHeaderExtension(EPackageExtension::Asset);
+	UPackage* TestRoot = NewObject<UPackage>(nullptr, TEXT("TestPackage"));
+	FTestLinkerLoad* TestLinker = new FTestLinkerLoad(TestRoot, TestPath);
+
+	// Good cases
+	const FStringView TestSingleSubobject = TEXT("/Script/SomeClass /Path/To/A/Package.Object:Subobject1");
+	const FStringView TestTwoSubobjects = TEXT("/Script/SomeClass /Path/To/A/Package.Object:Subobject1.Subobject2");
+	const FStringView TestNoSubobjects = TEXT("/Script/SomeClass /Path/To/A/Package.Object");
+	const FStringView TestTwoSubobjectsAndNoClassPath = TEXT("/Path/To/A/Package.Object:Subobject1.Subobject2");
+	const FStringView TestPackage = TEXT("/Script/SomeClass /Path/To/A/Package");
+
+	// Suspicious cases
+	const FStringView TestMissingSubobject = TEXT("/Script/SomeClass /Path/To/A/Package.Object:");
+	const FStringView TestMissingSubobjectWithTrailingDot = TEXT("/Script/SomeClass /Path/To/A/Package.Object:.");
+	const FStringView TestValidSubobjectWithTrailingDot = TEXT("/Script/SomeClass /Path/To/A/Package.Object:Subobject1.");
+	const FStringView TestTwoValidSubobjectsWithTrailingDot = TEXT("/Script/SomeClass /Path/To/A/Package.Object:Subobject1.Subobject2.");
+	const FStringView TestSingleMissingSubobject = TEXT("/Script/SomeClass /Path/To/A/Package.Object:BadSubobject1");
+	const FStringView TestTwoMissingSubobjects = TEXT("/Script/SomeClass /Path/To/A/Package.Object:BadSubobject1.BadSubobject2");
+	const FStringView TestOneValidOneMissingSubobject = TEXT("/Script/SomeClass /Path/To/A/Package.Object:Subobject1.BadSubobject2");
+	const FStringView TestMissingRootObject = TEXT("/Script/SomeClass /Path/To/A/Package.BadObject:Subobject1.Subobject2");
+	const FStringView TestMissingPackage = TEXT("/Script/SomeClass /Path/To/A/BadPackage.Object:Subobject1.Subobject2");
+
+	SECTION("Single subobject verification")
+	{
+		FPackageIndex Result;
+		REQUIRE(TestLinker->FindImport(TestSingleSubobject, Result));
+		REQUIRE(Result.IsImport());
+		REQUIRE(Result.ToImport() == TestLinker->SubObject1Index);
+	}
+
+	SECTION("Two subobjects verification")
+	{
+		FPackageIndex Result;
+		REQUIRE(TestLinker->FindImport(TestTwoSubobjects, Result));
+		REQUIRE(Result.IsImport());
+		REQUIRE(Result.ToImport() == TestLinker->SubObject2Index);
+	}
+
+	SECTION("No subobjects verification")
+	{
+		FPackageIndex Result;
+		REQUIRE(TestLinker->FindImport(TestNoSubobjects, Result));
+		REQUIRE(Result.IsImport());
+		REQUIRE(Result.ToImport() == TestLinker->RootObjectIndex);
+	}
+
+	SECTION("No class path verification")
+	{
+		FPackageIndex Result;
+		REQUIRE(TestLinker->FindImport(TestTwoSubobjectsAndNoClassPath, Result));
+		REQUIRE(Result.IsImport());
+		REQUIRE(Result.ToImport() == TestLinker->SubObject2Index);
+	}
+
+	SECTION("Package verification")
+	{
+		FPackageIndex Result;
+		REQUIRE(TestLinker->FindImport(TestPackage, Result));
+		REQUIRE(Result.IsImport());
+		REQUIRE(Result.ToImport() == TestLinker->PackageIndex);
+	}
+
+	SECTION("Trailing colon still finds root object")
+	{
+		FPackageIndex Result;
+		REQUIRE(TestLinker->FindImport(TestMissingSubobject, Result));
+		REQUIRE(Result.IsImport());
+		REQUIRE(Result.ToImport() == TestLinker->RootObjectIndex);
+	}
+
+	SECTION("Trailing colon and dot still finds root object")
+	{
+		FPackageIndex Result;
+		REQUIRE(TestLinker->FindImport(TestMissingSubobjectWithTrailingDot, Result));
+		REQUIRE(Result.IsImport());
+		REQUIRE(Result.ToImport() == TestLinker->RootObjectIndex);
+	}
+
+	SECTION("Trailing dot still finds subobject")
+	{
+		FPackageIndex Result;
+		REQUIRE(TestLinker->FindImport(TestValidSubobjectWithTrailingDot, Result));
+		REQUIRE(Result.IsImport());
+		REQUIRE(Result.ToImport() == TestLinker->SubObject1Index);
+	}
+
+	SECTION("Trailing dot still finds two subobjects")
+	{
+		FPackageIndex Result;
+		REQUIRE(TestLinker->FindImport(TestTwoValidSubobjectsWithTrailingDot, Result));
+		REQUIRE(Result.IsImport());
+		REQUIRE(Result.ToImport() == TestLinker->SubObject2Index);
+	}
+
+	SECTION("Single missing subobject not found")
+	{
+		FPackageIndex Result;
+		REQUIRE_FALSE(TestLinker->FindImport(TestSingleMissingSubobject, Result));
+		REQUIRE(Result.IsNull());
+	}
+
+	SECTION("Two missing subobjects not found")
+	{
+		FPackageIndex Result;
+		REQUIRE_FALSE(TestLinker->FindImport(TestTwoMissingSubobjects, Result));
+		REQUIRE(Result.IsNull());
+	}
+
+	SECTION("One valid and one missing subobject not found")
+	{
+		FPackageIndex Result;
+		REQUIRE_FALSE(TestLinker->FindImport(TestOneValidOneMissingSubobject, Result));
+		REQUIRE(Result.IsNull());
+	}
+
+	SECTION("Missing root object not found")
+	{
+		FPackageIndex Result;
+		REQUIRE_FALSE(TestLinker->FindImport(TestMissingRootObject, Result));
+		REQUIRE(Result.IsNull());
+	}
+
+	SECTION("Missing package not found")
+	{
+		FPackageIndex Result;
+		REQUIRE_FALSE(TestLinker->FindImport(TestMissingRootObject, Result));
+		REQUIRE(Result.IsNull());
+	}
+
+	FUObjectThreadContext::Get().IsDeletingLinkers = true;
+	delete TestLinker;
+	FUObjectThreadContext::Get().IsDeletingLinkers = false;
+}
diff --git a/Engine/Source/Programs/LowLevelTests/FoundationTests/Tests/PackageNameTests.cpp b/Engine/Source/Programs/LowLevelTests/FoundationTests/Tests/PackageNameTests.cpp
new file mode 100644
--- /dev/null
+++ b/Engine/Source/Programs/LowLevelTests/FoundationTests/Tests/PackageNameTests.cpp
@@ -0,0 +1,143 @@
+// Copyright Epic Games, Inc. All Rights Reserved.
+
+#include "CoreTypes.h"
+
+#include "Misc/PackageName.h"
+#include "TestHarness.h"
+
+TEST_CASE("Split full object path with subobjects", "[CoreUObject][FPackageName::SplitFullObjectPath]")
+{
+	const FStringView ExpectedClassPath = TEXT("/Script/SomePackage.SomeClass");
+	const FStringView ExpectedPackagePath = TEXT("/Path/To/A/Package");
+	const FStringView ExpectedObjectName = TEXT("Object");
+	const FStringView ExpectedSubobject1Name = TEXT("Subobject1");
+	const FStringView ExpectedSubobject2Name = TEXT("Subobject2");
+
+	// Good cases
+	const FStringView TestSingleSubobject = TEXT("/Script/SomePackage.SomeClass /Path/To/A/Package.Object:Subobject1");
+	const FStringView TestTwoSubobjects = TEXT("/Script/SomePackage.SomeClass /Path/To/A/Package.Object:Subobject1.Subobject2");
+	const FStringView TestNoSubobjects = TEXT("/Script/SomePackage.SomeClass /Path/To/A/Package.Object");
+	const FStringView TestTwoSubobjectsAndNoClassPath = TEXT("/Path/To/A/Package.Object:Subobject1.Subobject2");
+	const FStringView TestPackage = TEXT("/Script/SomePackage.SomeClass /Path/To/A/Package");
+
+	// Suspicious cases
+	const FStringView TestMissingSubobject = TEXT("/Script/SomePackage.SomeClass /Path/To/A/Package.Object:");
+	const FStringView TestMissingSubobjectWithTrailingDot = TEXT("/Script/SomePackage.SomeClass /Path/To/A/Package.Object:.");
+	const FStringView TestValidSubobjectWithTrailingDot = TEXT("/Script/SomePackage.SomeClass /Path/To/A/Package.Object:Subobject1.");
+	const FStringView TestTwoValidSubobjectsWithTrailingDot = TEXT("/Script/SomePackage.SomeClass /Path/To/A/Package.Object:Subobject1.Subobject2.");
+
+	FStringView ClassPath;
+	FStringView PackagePath;
+	FStringView ObjectName; 
+	TArray<FStringView> SubobjectNames;
+
+	SECTION("Single subobject verification")
+	{
+		FPackageName::SplitFullObjectPath(TestSingleSubobject, ClassPath, PackagePath, ObjectName, SubobjectNames);
+
+		REQUIRE(ClassPath == ExpectedClassPath);
+		REQUIRE(PackagePath == ExpectedPackagePath);
+		REQUIRE(ObjectName == ExpectedObjectName);
+		REQUIRE(SubobjectNames.Num() == 1);
+		REQUIRE(SubobjectNames[0] == ExpectedSubobject1Name);
+	}
+
+	SECTION("Two subobjects verification")
+	{
+		FPackageName::SplitFullObjectPath(TestTwoSubobjects, ClassPath, PackagePath, ObjectName, SubobjectNames);
+
+		REQUIRE(ClassPath == ExpectedClassPath);
+		REQUIRE(PackagePath == ExpectedPackagePath);
+		REQUIRE(ObjectName == ExpectedObjectName);
+		REQUIRE(SubobjectNames.Num() == 2);
+		REQUIRE(SubobjectNames[0] == ExpectedSubobject1Name);
+		REQUIRE(SubobjectNames[1] == ExpectedSubobject2Name);
+	}
+
+	SECTION("No subobjects verification")
+	{
+		FPackageName::SplitFullObjectPath(TestNoSubobjects, ClassPath, PackagePath, ObjectName, SubobjectNames);
+
+		REQUIRE(ClassPath == ExpectedClassPath);
+		REQUIRE(PackagePath == ExpectedPackagePath);
+		REQUIRE(ObjectName == ExpectedObjectName);
+		REQUIRE(SubobjectNames.Num() == 0);
+	}
+
+	SECTION("No class path verification (bDetectClassName on)")
+	{
+		FPackageName::SplitFullObjectPath(TestTwoSubobjectsAndNoClassPath, ClassPath, PackagePath, ObjectName, SubobjectNames);
+
+		REQUIRE(ClassPath.Len() == 0);
+		REQUIRE(PackagePath == ExpectedPackagePath);
+		REQUIRE(ObjectName == ExpectedObjectName);
+		REQUIRE(SubobjectNames.Num() == 2);
+		REQUIRE(SubobjectNames[0] == ExpectedSubobject1Name);
+		REQUIRE(SubobjectNames[1] == ExpectedSubobject2Name);
+	}
+
+	SECTION("No class path verification (bDetectClassName off)")
+	{
+		FPackageName::SplitFullObjectPath(TestTwoSubobjectsAndNoClassPath, ClassPath, PackagePath, ObjectName, SubobjectNames, false);
+
+		REQUIRE(ClassPath.Len() == 0);
+		REQUIRE(PackagePath == ExpectedPackagePath);
+		REQUIRE(ObjectName == ExpectedObjectName);
+		REQUIRE(SubobjectNames.Num() == 2);
+		REQUIRE(SubobjectNames[0] == ExpectedSubobject1Name);
+		REQUIRE(SubobjectNames[1] == ExpectedSubobject2Name);
+	}
+
+	SECTION("Package verification")
+	{
+		FPackageName::SplitFullObjectPath(TestPackage, ClassPath, PackagePath, ObjectName, SubobjectNames);
+
+		REQUIRE(ClassPath == ExpectedClassPath);
+		REQUIRE(PackagePath == ExpectedPackagePath);
+		REQUIRE(ObjectName == TEXT(""));
+		REQUIRE(SubobjectNames.Num() == 0);
+	}
+
+	SECTION("Missing subobject name yields empty subobjects array")
+	{
+		FPackageName::SplitFullObjectPath(TestMissingSubobject, ClassPath, PackagePath, ObjectName, SubobjectNames);
+
+		REQUIRE(ClassPath == ExpectedClassPath);
+		REQUIRE(PackagePath == ExpectedPackagePath);
+		REQUIRE(ObjectName == ExpectedObjectName);
+		REQUIRE(SubobjectNames.Num() == 0);
+	}
+
+	SECTION("Missing subobject name with trailing dot yields empty subobjects array")
+	{
+		FPackageName::SplitFullObjectPath(TestMissingSubobjectWithTrailingDot, ClassPath, PackagePath, ObjectName, SubobjectNames);
+
+		REQUIRE(ClassPath == ExpectedClassPath);
+		REQUIRE(PackagePath == ExpectedPackagePath);
+		REQUIRE(ObjectName == ExpectedObjectName);
+		REQUIRE(SubobjectNames.Num() == 0);
+	}
+
+	SECTION("Valid subobject with trailing dot still reports correct subobject name")
+	{
+		FPackageName::SplitFullObjectPath(TestValidSubobjectWithTrailingDot, ClassPath, PackagePath, ObjectName, SubobjectNames);
+
+		REQUIRE(ClassPath == ExpectedClassPath);
+		REQUIRE(PackagePath == ExpectedPackagePath);
+		REQUIRE(ObjectName == ExpectedObjectName);
+		REQUIRE(SubobjectNames.Num() == 1);
+		REQUIRE(SubobjectNames[0] == ExpectedSubobject1Name);
+	}
+
+	SECTION("Two valid subobjects with trailing dot still reports correct subobject names")
+	{
+		FPackageName::SplitFullObjectPath(TestTwoValidSubobjectsWithTrailingDot, ClassPath, PackagePath, ObjectName, SubobjectNames);
+
+		REQUIRE(ClassPath == ExpectedClassPath);
+		REQUIRE(PackagePath == ExpectedPackagePath);
+		REQUIRE(ObjectName == ExpectedObjectName);
+		REQUIRE(SubobjectNames.Num() == 2);
+		REQUIRE(SubobjectNames[0] == ExpectedSubobject1Name);
+		REQUIRE(SubobjectNames[1] == ExpectedSubobject2Name);
+	}
+}
\ No newline at end of file
diff --git a/Engine/Source/Runtime/CoreUObject/Private/Blueprint/BlueprintSupport.cpp b/Engine/Source/Runtime/CoreUObject/Private/Blueprint/BlueprintSupport.cpp
--- a/Engine/Source/Runtime/CoreUObject/Private/Blueprint/BlueprintSupport.cpp
+++ b/Engine/Source/Runtime/CoreUObject/Private/Blueprint/BlueprintSupport.cpp
@@ -2541,14 +2541,24 @@ UObject* FLinkerLoad::RequestPlaceholderValue(const FProperty* Property, const U
 				}
 			}
 
-			const FString ObjectPathStr(ObjectPath);
+			const FStringView ObjectPathStr(ObjectPath);
 			// we don't need placeholders for native object references and for non-BP class objects (the 
 			// calling code should properly handle null return values)
 			if (!FPackageName::IsScriptPackage(ObjectPathStr) && ObjectType->HasAnyClassFlags(CLASS_NeedsDeferredDependencyLoading))
 			{
-				const FString ObjectName = FPackageName::ObjectPathToObjectName(ObjectPathStr);
-				Placeholder = MakeImportPlaceholder<ULinkerPlaceholderClass>(LinkerRoot, ObjectType, *ObjectName);
-				ImportPlaceholders.Add(ObjId, Placeholder);
+				FPackageIndex ImportIndex;
+
+				if (FindImport(ObjectPathStr, ImportIndex))
+				{
+					const FStringView ObjectName = FPackageName::ObjectPathToObjectName(ObjectPathStr);
+
+					Placeholder = MakeImportPlaceholder<ULinkerPlaceholderClass>(LinkerRoot, ObjectType, ObjectName.GetData(), ImportIndex.ToImport());
+					ImportPlaceholders.Add(ObjId, Placeholder);
+				}
+				else
+				{
+					UE_LOG(LogBlueprintSupport, Error, TEXT("'%s' was not found in the import table for package '%s'."), ObjectPathStr.GetData(), *LinkerRoot->GetFullName());
+				}
 			}
 		}
 	}
diff --git a/Engine/Source/Runtime/CoreUObject/Private/Misc/PackageName.cpp b/Engine/Source/Runtime/CoreUObject/Private/Misc/PackageName.cpp
--- a/Engine/Source/Runtime/CoreUObject/Private/Misc/PackageName.cpp
+++ b/Engine/Source/Runtime/CoreUObject/Private/Misc/PackageName.cpp
@@ -1146,6 +1146,65 @@ void FPackageName::SplitFullObjectPath(FStringView InFullObjectPath, FStringView
 	}
 }
 
+void FPackageName::SplitFullObjectPath(const FString& InFullObjectPath, FString& OutClassName,
+	FString& OutPackageName, FString& OutObjectName, TArray<FString>& OutSubobjectNames, bool bDetectClassName)
+{
+	FStringView ClassName;
+	FStringView PackageName;
+	FStringView ObjectName;
+	TArray<FStringView> SubObjectNames;
+	SplitFullObjectPath(InFullObjectPath, ClassName, PackageName, ObjectName, SubObjectNames, bDetectClassName);
+	OutClassName = ClassName;
+	OutPackageName = PackageName;
+	OutObjectName = ObjectName;
+	OutSubobjectNames.Empty();
+	OutSubobjectNames.Append(SubObjectNames);
+}
+
+void FPackageName::SplitFullObjectPath(FStringView InFullObjectPath, FStringView& OutClassName,
+	FStringView& OutPackageName, FStringView& OutObjectName, TArray<FStringView>& OutSubobjectNames, bool bDetectClassName)
+{
+	FStringView FullSubobjectString;
+	SplitFullObjectPath(InFullObjectPath, OutClassName, OutPackageName, OutObjectName, FullSubobjectString, bDetectClassName);
+
+	if (FullSubobjectString.Len() > 0)
+	{
+		OutSubobjectNames.Empty();
+
+		auto ExtractBeforeDotDelim = [&FullSubobjectString](FStringView& OutStringView)
+		{
+			int32 DelimIndex;
+			if (FullSubobjectString.FindChar(TCHAR('.'), DelimIndex))
+			{
+				OutStringView = FullSubobjectString.Left(DelimIndex);
+				FullSubobjectString.RightChopInline(DelimIndex + 1);
+				return true;
+			}
+			else
+			{
+				OutStringView.Reset();
+				return false;
+			}
+		};
+
+		bool bFoundDelim = false;
+		do
+		{
+			FStringView SubobjectName;
+			bFoundDelim = ExtractBeforeDotDelim(SubobjectName);
+			const bool bSubobjectNameIsNotEmpty = (SubobjectName.Len() > 0);
+			if (bFoundDelim && bSubobjectNameIsNotEmpty)
+			{
+				OutSubobjectNames.Add(SubobjectName);
+			}
+			else if (FullSubobjectString.Len() > 0)
+			{
+				OutSubobjectNames.Add(FullSubobjectString);
+			}
+		} while (bFoundDelim);
+	}
+}
+
 FString FPackageName::GetLongPackageAssetName(const FString& InLongPackageName)
 {
 	return GetShortName(InLongPackageName);
diff --git a/Engine/Source/Runtime/CoreUObject/Private/UObject/LinkerLoad.cpp b/Engine/Source/Runtime/CoreUObject/Private/UObject/LinkerLoad.cpp
--- a/Engine/Source/Runtime/CoreUObject/Private/UObject/LinkerLoad.cpp
+++ b/Engine/Source/Runtime/CoreUObject/Private/UObject/LinkerLoad.cpp
@@ -6389,6 +6389,76 @@ bool FLinkerLoad::FindImport(FPackageIndex OuterIndex, FName ObjectName, FPackag
 	return false;
 }
 
+bool FLinkerLoad::FindImport(FStringView FullObjectPath, FPackageIndex& OutObjectIndex)
+{
+	bool bIsValid = false;
+
+	FStringView ClassName;
+	FStringView PackageName;
+	FStringView ObjectName;
+	TArray<FStringView> SubobjectNames;
+
+	FPackageName::SplitFullObjectPath(FullObjectPath, ClassName, PackageName, ObjectName, SubobjectNames);
+
+	FName PackageFName(PackageName);
+	FName ObjectFName(ObjectName);
+
+	FPackageIndex PackageIndex;
+
+	bIsValid = FindImportPackage(PackageFName, PackageIndex);
+
+	if (bIsValid)
+	{
+		const bool bHasRootObject = !ObjectFName.IsNone();
+
+		if (bHasRootObject)
+		{
+			FPackageIndex ObjectIndex;
+
+			bIsValid = FindImport(PackageIndex, ObjectFName, ObjectIndex);
+
+			if (bIsValid)
+			{
+				const bool bHasSubobjects = (SubobjectNames.Num() > 0);
+
+				if (bHasSubobjects)
+				{
+					FPackageIndex CurrentOuterIndex = ObjectIndex;
+					FPackageIndex SubobjectIndex;
+
+					for (FStringView SubobjectName : SubobjectNames)
+					{
+						if (FindImport(CurrentOuterIndex, FName(SubobjectName), SubobjectIndex))
+						{
+							CurrentOuterIndex = SubobjectIndex;
+						}
+						else
+						{
+							bIsValid = false;
+							break;
+						}
+					}
+
+					if (bIsValid)
+					{
+						OutObjectIndex = SubobjectIndex;
+					}
+				}
+				else
+				{
+					OutObjectIndex = ObjectIndex;
+				}
+			}
+		}
+		else
+		{
+			OutObjectIndex = PackageIndex;
+		}
+	}
+
+	return bIsValid;
+}
+
 /**
  * Locates the class adjusted index and its package adjusted index for a given class name in the import map
  */
diff --git a/Engine/Source/Runtime/CoreUObject/Public/Misc/PackageName.h b/Engine/Source/Runtime/CoreUObject/Public/Misc/PackageName.h
--- a/Engine/Source/Runtime/CoreUObject/Public/Misc/PackageName.h
+++ b/Engine/Source/Runtime/CoreUObject/Public/Misc/PackageName.h
@@ -184,13 +184,13 @@ public:
 	static COREUOBJECT_API bool SplitLongPackageName(const FString& InLongPackageName, FString& OutPackageRoot, FString& OutPackagePath, FString& OutPackageName, const bool bStripRootLeadingSlash = false);
 
 	/**
-	 * Split a full object path (Class /Path/To/A/Package.Object:SubObject) into its constituent pieces
+	 * Split a full object path (Class /Path/To/A/Package.Object:Subobject1.Subobject2) into its constituent pieces
 	 *  
 	 * @param InFullObjectPath  Full object path we want to split
 	 * @param OutClassName      The extracted class name (Class)
 	 * @param OutPackageName    The extracted package name (/Path/To/A/Package)
 	 * @param OutObjectName     The extracted object name (Object)
-	 * @param OutSubObjectName  The extracted subobject name (SubObject)
+	 * @param OutSubObjectName  The extracted subobject name (Subobject1.Subobject2) - Note: nested subobjects are not split
 	 * @param bDetectClassName  If true, the optional Class will be detected and separated based on a space.
 	 *                          If false, and there is a space, the space and text before it will be included in the
 	 *                          other names. Spaces in those names is invalid, but some code ignores the
@@ -201,6 +201,25 @@ public:
 	static COREUOBJECT_API void SplitFullObjectPath(FStringView InFullObjectPath, FStringView& OutClassName,
 		FStringView& OutPackageName, FStringView& OutObjectName, FStringView& OutSubObjectName, bool bDetectClassName=true);
 
+	/**
+	 * Split a full object path (Class /Path/To/A/Package.Object:Subobject1.Subobject2) into its constituent pieces.
+	 * All subobjects are split individually and stored in an array.
+	 *
+	 * @param InFullObjectPath  Full object path we want to split
+	 * @param OutClassName      The extracted class name (Class)
+	 * @param OutPackageName    The extracted package name (/Path/To/A/Package)
+	 * @param OutSubobjectNames The extracted object name (Object)
+	 * @param OutSubObjectName  The extracted subobject names (Subobject1 and Subobject2)
+	 * @param bDetectClassName  If true, the optional Class will be detected and separated based on a space.
+	 *                          If false, and there is a space, the space and text before it will be included in the
+	 *                          other names. Spaces in those names is invalid, but some code ignores the
+	 *                          invalidity in ObjectName if it only cares about packageName.
+	 */
+	static COREUOBJECT_API void SplitFullObjectPath(const FString& InFullObjectPath, FString& OutClassName,
+		FString& OutPackageName, FString& OutObjectName, TArray<FString>& OutSubobjectNames, bool bDetectClassName = true);
+	static COREUOBJECT_API void SplitFullObjectPath(FStringView InFullObjectPath, FStringView& OutClassName,
+		FStringView& OutPackageName, FStringView& OutObjectName, TArray<FStringView>& OutSubobjectNames, bool bDetectClassName = true);
+
 	/** 
 	 * Returns true if the path starts with a valid root (i.e. /Game/, /Engine/, etc) and contains no illegal characters.
 	 *
diff --git a/Engine/Source/Runtime/CoreUObject/Public/UObject/LinkerLoad.h b/Engine/Source/Runtime/CoreUObject/Public/UObject/LinkerLoad.h
--- a/Engine/Source/Runtime/CoreUObject/Public/UObject/LinkerLoad.h
+++ b/Engine/Source/Runtime/CoreUObject/Public/UObject/LinkerLoad.h
@@ -115,7 +115,7 @@ struct FScopedCreateImportCounter
 class FAsyncArchive;
 
 class FLinkerLoad 
-#if !WITH_EDITOR
+#if !WITH_EDITOR && !WITH_LOW_LEVEL_TESTS
 	final 
 #endif
 	: public FLinker, public FArchiveUObject
@@ -572,18 +572,28 @@ public:
 	 * Locates package index for a UPackage import
 	 */
 	COREUOBJECT_API bool FindImportPackage(FName PackageName, FPackageIndex& PackageIdx);
-	/* Locates package index for a given name in an outer. */
+
+	/** 
+	 * Locates package index for a given name in an outer.
+	 */
 	COREUOBJECT_API bool FindImport(FPackageIndex OuterIndex, FName ObjectName, FPackageIndex& OutObjectIndex);
 
+	/**
+	 * Finds an import given the full object path as a string.
+	 * 
+	 * Note: since imports are stored in an array, it needs to linearly search for _each_ import in the object path.
+	 */
+	COREUOBJECT_API bool FindImport(FStringView FullObjectPath, FPackageIndex& OutObjectIndex);
+
 	/**
 	 * Locates the class adjusted index and its package adjusted index for a given class name in the import map
 	 */
-	COREUOBJECT_API bool FindImportClassAndPackage( FName ClassName, FPackageIndex& ClassIdx, FPackageIndex& PackageIdx );
+	COREUOBJECT_API bool FindImportClassAndPackage(FName ClassName, FPackageIndex& ClassIdx, FPackageIndex& PackageIdx);
 	
 	/**
 	 * Attempts to find the index for the given class object in the import list and adds it + its package if it does not exist
 	 */
-	COREUOBJECT_API bool CreateImportClassAndPackage( FName ClassName, FName PackageName, FPackageIndex& ClassIdx, FPackageIndex& PackageIdx );
+	COREUOBJECT_API bool CreateImportClassAndPackage(FName ClassName, FName PackageName, FPackageIndex& ClassIdx, FPackageIndex& PackageIdx);
 
 	/**
 	 * Allows object instances to be converted to other classes upon loading a package
