[StateStream]
* Added so FStateStreamCreator takes in statestream id and state stream creators are sorted. This makes it easier to access "more core system" statestreams

diff --git a/Engine/Source/Programs/LowLevelTests/StateStreamTests/Tests/StateStreamTests.cpp b/Engine/Source/Programs/LowLevelTests/StateStreamTests/Tests/StateStreamTests.cpp
--- a/Engine/Source/Programs/LowLevelTests/StateStreamTests/Tests/StateStreamTests.cpp
+++ b/Engine/Source/Programs/LowLevelTests/StateStreamTests/Tests/StateStreamTests.cpp
@@ -1,5 +1,6 @@
 // Copyright Epic Games, Inc. All Rights Reserved.
 #include "TestCommon/Initialization.h"
+#include "StateStreamCreator.h"
 #include "TestHarness.h"
 #include "MockStateStreamImpl.h"
 #include "StateStreamManagerImpl.h"
@@ -570,3 +571,20 @@ TEST_CASE("StateStream::Interpolation", "[StateStream]")
 		CHECK_EQUAL((FMath::Abs(Object->GetInfo().WorldTransform.GetTranslation().X - 10.0f) < 0.001f), true);
 	}
 }
+
+TEST_CASE("StateStreamCreator", "[StateStream]")
+{
+	uint32 Counter = 0;
+
+	FStateStreamCreator Creator1(1, [&](const FStateStreamRegisterContext&){ CHECK_EQUAL(Counter++, 1); }, [&](const FStateStreamUnregisterContext&){ --Counter; });
+	FStateStreamCreator Creator2(3, [&](const FStateStreamRegisterContext&){ CHECK_EQUAL(Counter++, 3); }, [&](const FStateStreamUnregisterContext&){ --Counter; });
+	FStateStreamCreator Creator3(2, [&](const FStateStreamRegisterContext&){ CHECK_EQUAL(Counter++, 2); }, [&](const FStateStreamUnregisterContext&){ --Counter; });
+	FStateStreamCreator Creator4(0, [&](const FStateStreamRegisterContext&){ CHECK_EQUAL(Counter++, 0); }, [&](const FStateStreamUnregisterContext&){ --Counter; });
+	FStateStreamCreator Creator5(4, [&](const FStateStreamRegisterContext&){ CHECK_EQUAL(Counter++, 4); }, [&](const FStateStreamUnregisterContext&){ --Counter; });
+
+	FStateStreamManagerImpl ManagerImpl;
+	FStateStreamCreator::RegisterStateStreams(FStateStreamRegisterContext{ManagerImpl});
+	CHECK_EQUAL(Counter, 5);
+	FStateStreamCreator::UnregisterStateStreams(FStateStreamUnregisterContext{ManagerImpl});
+	CHECK_EQUAL(Counter, 0);
+}
diff --git a/Engine/Source/Runtime/StateStream/Private/StateStreamManagerImpl.cpp b/Engine/Source/Runtime/StateStream/Private/StateStreamManagerImpl.cpp
--- a/Engine/Source/Runtime/StateStream/Private/StateStreamManagerImpl.cpp
+++ b/Engine/Source/Runtime/StateStream/Private/StateStreamManagerImpl.cpp
@@ -146,22 +146,33 @@ FStateStreamManagerImpl::~FStateStreamManagerImpl()
 
 FStateStreamCreator* FStateStreamCreator::First;
 
-FStateStreamCreator::FStateStreamCreator(FRegisterFunction&& InRegisterFunction, FUnregisterFunction&& InUnregisterFunction)
+FStateStreamCreator::FStateStreamCreator(uint32 Id_, FRegisterFunction&& InRegisterFunction, FUnregisterFunction&& InUnregisterFunction)
 :	RegisterFunction(MoveTemp(InRegisterFunction))
 ,	UnregisterFunction(MoveTemp(InUnregisterFunction))
+,	Id(Id_)
 {
-	if (!First)
+	Prev = nullptr;
+	for (FStateStreamCreator* It=First; It; It=It->Next)
+	{
+		if (Id < It->Id)
+			break;
+		Prev = It;
+	}
+
+	if (!Prev)
 	{
+		Next = First;
 		First = this;
-		Prev = nullptr;
-		Next = nullptr;
 	}
 	else
 	{
-		Prev = nullptr;
-		Next = First;
+		Next = Prev->Next;
+		Prev->Next = this;
+	}
+
+	if (Next)
+	{
 		Next->Prev = this;
-		First = this;
 	}
 }
 
diff --git a/Engine/Source/Runtime/StateStream/Public/StateStreamCreator.h b/Engine/Source/Runtime/StateStream/Public/StateStreamCreator.h
--- a/Engine/Source/Runtime/StateStream/Public/StateStreamCreator.h
+++ b/Engine/Source/Runtime/StateStream/Public/StateStreamCreator.h
@@ -31,7 +31,7 @@ class FStateStreamCreator
 public:
 	using FRegisterFunction = TFunction<void(const FStateStreamRegisterContext&)>;
 	using FUnregisterFunction = TFunction<void(const FStateStreamUnregisterContext&)>;
-	UE_API FStateStreamCreator(FRegisterFunction&& InRegisterFunction, FUnregisterFunction&& InUnregisterFunction);
+	UE_API FStateStreamCreator(uint32 Id, FRegisterFunction&& InRegisterFunction, FUnregisterFunction&& InUnregisterFunction);
 	UE_API ~FStateStreamCreator();
 
 	UE_API static void RegisterStateStreams(const FStateStreamRegisterContext&);
@@ -47,6 +47,7 @@ private:
 	static FStateStreamCreator* First;
 	FStateStreamCreator* Next;
 	FStateStreamCreator* Prev;
+	uint32 Id;
 };
 
 ////////////////////////////////////////////////////////////////////////////////////////////////////
