[shaders] only add shadercompilercommon to foundationtests for win64/mac/linux platforms since shaderconductor & deps isn't supported on any other platforms

diff --git a/Engine/Source/Programs/LowLevelTests/FoundationTests/FoundationTests.Build.cs b/Engine/Source/Programs/LowLevelTests/FoundationTests/FoundationTests.Build.cs
--- a/Engine/Source/Programs/LowLevelTests/FoundationTests/FoundationTests.Build.cs
+++ b/Engine/Source/Programs/LowLevelTests/FoundationTests/FoundationTests.Build.cs
@@ -56,9 +56,19 @@ public class FoundationTests : TestModuleRules
 			PrivateDependencyModuleNames.AddRange(
 				new string[] {
 					"DesktopPlatform",
-					"ShaderCompilerCommon",
 					"ShaderPreprocessor"
 				});
 		}
+
+		if (Target.bBuildWithEditorOnlyData &&
+			(Target.Platform == UnrealTargetPlatform.Mac ||
+			 Target.Platform == UnrealTargetPlatform.Win64 ||
+			 Target.Platform == UnrealTargetPlatform.Linux))
+		{
+			PrivateDependencyModuleNames.AddRange(
+				new string[] {
+					"ShaderCompilerCommon",
+				});
+		}
 	}
 }
\ No newline at end of file
