Fix Mac failure at exit and MathCoreTests IOS

#rnx

diff --git a/Engine/Source/Developer/LowLevelTestsRunner/Private/Platform/Mac/MacTestRunner.cpp b/Engine/Source/Developer/LowLevelTestsRunner/Private/Platform/Mac/MacTestRunner.cpp
--- a/Engine/Source/Developer/LowLevelTestsRunner/Private/Platform/Mac/MacTestRunner.cpp
+++ b/Engine/Source/Developer/LowLevelTestsRunner/Private/Platform/Mac/MacTestRunner.cpp
@@ -40,6 +40,14 @@ int main(int argc, const char* argv[])
 	
 	int Result = Helper.Result;
 	[Helper release];
+
+	if (GGameThreadId != FPlatformTLS::GetCurrentThreadId())
+	{
+		// Force Game thread to be the current one for proper cleanup at exit
+		FTaskTagScope Scope(ETaskTag::EGameThread);
+		GGameThreadId = FPlatformTLS::GetCurrentThreadId();
+	}
+
 	return Result;
 }
 #endif
diff --git a/Engine/Source/Programs/LowLevelTests/MathCoreTests/MathCoreTests.Build.cs b/Engine/Source/Programs/LowLevelTests/MathCoreTests/MathCoreTests.Build.cs
--- a/Engine/Source/Programs/LowLevelTests/MathCoreTests/MathCoreTests.Build.cs
+++ b/Engine/Source/Programs/LowLevelTests/MathCoreTests/MathCoreTests.Build.cs
@@ -36,9 +36,14 @@ public class MathCoreTests : TestModuleRules
 	public MathCoreTests(ReadOnlyTargetRules Target) : base(Target, InUsesCatch2:true)
 	{
 		PrivateDependencyModuleNames.AddRange(
-			new string[] {
+			[
 				"Core",
 				"MathCore",
-			});
+			]);
+
+		if (Target.Platform == UnrealTargetPlatform.IOS)
+		{
+			PrivateDependencyModuleNames.AddRange([ "ApplicationCore", "CoreUObject" ]);
+		}
 	}
 }
\ No newline at end of file
