[FoundationTests]
* Added bCompileAgainstApplicationCore = true since slate is pulled in which uses application core

diff --git a/Engine/Source/Programs/LowLevelTests/FoundationTests/FoundationTests.Target.cs b/Engine/Source/Programs/LowLevelTests/FoundationTests/FoundationTests.Target.cs
--- a/Engine/Source/Programs/LowLevelTests/FoundationTests/FoundationTests.Target.cs
+++ b/Engine/Source/Programs/LowLevelTests/FoundationTests/FoundationTests.Target.cs
@@ -13,6 +13,8 @@ public class FoundationTestsTarget : TestTargetRules
 
 		bCompileAgainstCoreUObject = true;
 
+		bCompileAgainstApplicationCore = true;
+
 		bBuildWithEditorOnlyData = Target.Platform.IsInGroup(UnrealPlatformGroup.Desktop)
 			&& (Target.Configuration == UnrealTargetConfiguration.Debug || Target.Configuration == UnrealTargetConfiguration.Development);
 	}
