More basic UObject declarations to test changes to UHT/Object init

diff --git a/Engine/Source/Programs/LowLevelTests/FoundationTests/Private/TestInterface.h b/Engine/Source/Programs/LowLevelTests/FoundationTests/Private/TestInterface.h
new file mode 100644
--- /dev/null
+++ b/Engine/Source/Programs/LowLevelTests/FoundationTests/Private/TestInterface.h
@@ -0,0 +1,24 @@
+// Copyright Epic Games, Inc. All Rights Reserved.
+
+#pragma once 
+
+#include "CoreTypes.h"
+#include "UObject/Interface.h"
+#include "UObject/Object.h"
+
+#include "TestInterface.generated.h"
+
+UINTERFACE(NotBlueprintable)
+class UTestInterface : public UInterface
+{
+	GENERATED_UINTERFACE_BODY()
+};
+
+class ITestInterface
+{
+	GENERATED_IINTERFACE_BODY()
+public:
+
+    UFUNCTION(BlueprintCallable, Category="Test")
+    virtual FName GetNumberedName(FName BaseName, int32 InNumber) = 0;
+};
\ No newline at end of file
diff --git a/Engine/Source/Programs/LowLevelTests/FoundationTests/Private/TestObject.cpp b/Engine/Source/Programs/LowLevelTests/FoundationTests/Private/TestObject.cpp
new file mode 100644
--- /dev/null
+++ b/Engine/Source/Programs/LowLevelTests/FoundationTests/Private/TestObject.cpp
@@ -0,0 +1,23 @@
+// Copyright Epic Games, Inc. All Rights Reserved.
+
+#include "TestObject.h"
+
+UTestInterface::UTestInterface(const FObjectInitializer& ObjectInitializer)
+    : Super(ObjectInitializer)
+{
+}
+
+UObject* UTestObject::GetStrongObjectReference()
+{
+    return StrongObjectReference;
+}
+
+int32 UTestObject::GetBPOverrideableValue_Implementation()
+{
+    return 17;
+}
+
+FName UTestObject::GetNumberedName(FName BaseName, int32 InNumber)
+{
+    return FName(BaseName, InNumber);
+}
\ No newline at end of file
diff --git a/Engine/Source/Programs/LowLevelTests/FoundationTests/Private/TestObject.h b/Engine/Source/Programs/LowLevelTests/FoundationTests/Private/TestObject.h
--- a/Engine/Source/Programs/LowLevelTests/FoundationTests/Private/TestObject.h
+++ b/Engine/Source/Programs/LowLevelTests/FoundationTests/Private/TestObject.h
@@ -1,6 +1,7 @@
 // Copyright Epic Games, Inc. All Rights Reserved.
 
 #include "CoreTypes.h"
+#include "TestInterface.h"
 #include "UObject/Object.h"
 
 #include "TestObject.generated.h"
@@ -28,16 +29,26 @@ struct FTestStruct
 };
 
 UCLASS()
-class UTestObject : public UObject
+class UTestObject : public UObject, public ITestInterface
 {
     GENERATED_BODY()
 
 public:
     UPROPERTY()
     TObjectPtr<UObject> StrongObjectReference;
-    
+
+#if WITH_EDITORONLY_DATA
+    UPROPERTY()
+    int32 EditorOnlyProp;
+#endif
+
     UPROPERTY()
     TWeakObjectPtr<UObject> WeakObjectReference;
+    
+#if WITH_VERSE_VM
+    UPROPERTY()
+    int32 VerseOnlyProp;
+#endif
 
     UPROPERTY()
     TSoftObjectPtr<UObject> SoftObjectReference;
@@ -45,9 +56,28 @@ public:
     UPROPERTY()
     FSoftObjectPath SoftObjectPath;
     
+#if WITH_EDITORONLY_DATA 
+#if WITH_VERSE_VM
+    UPROPERTY()
+    int32 EditorVerseOnlyProp;
+#endif
+#endif
+    
     UPROPERTY()
     FTestStruct EmbeddedStruct;
     
     UPROPERTY()
     TArray<FTestStruct> ArrayStructs;
+
+    UFUNCTION(BlueprintCallable, Category="Test")
+    UObject* GetStrongObjectReference();
+
+    UFUNCTION(BlueprintImplementableEvent)
+    void DoSomethingInBP();
+
+    UFUNCTION(BlueprintNativeEvent)
+    int32 GetBPOverrideableValue();
+
+    UFUNCTION(BlueprintCallable, Category="Test")
+    virtual FName GetNumberedName(FName BaseName, int32 InNumber) override;
 };
\ No newline at end of file
diff --git a/Engine/Source/Programs/LowLevelTests/FoundationTests/Tests/UObjectTests.cpp b/Engine/Source/Programs/LowLevelTests/FoundationTests/Tests/UObjectTests.cpp
new file mode 100644
--- /dev/null
+++ b/Engine/Source/Programs/LowLevelTests/FoundationTests/Tests/UObjectTests.cpp
@@ -0,0 +1,19 @@
+// Copyright Epic Games, Inc. All Rights Reserved.
+
+#include "CoreTypes.h"
+
+#include "TestHarness.h"
+#include "TestObject.h"
+#include "UObject/Class.h"
+#include "UObject/Package.h"
+
+TEST_CASE("CoreUObject::UFunction::Basic", "[CoreUObject]")
+{
+    FName PackageName = MakeUniqueObjectName(nullptr, UPackage::StaticClass(), FName("/Memory/UFunctionBasic"));
+    UPackage* Package = NewObject<UPackage>(nullptr, UPackage::StaticClass(), PackageName);
+    UTestObject* Obj = NewObject<UTestObject>(Package, "TestObject");
+
+    int32 Value = Obj->GetBPOverrideableValue();
+    int32 Value2 = Obj->GetBPOverrideableValue_Implementation();
+    CHECK_EQUAL(Value, Value2);
+}
\ No newline at end of file
