More basic UObject tests to test changes to UHT/Object init

diff --git a/Engine/Source/Programs/LowLevelTests/FoundationTests/FoundationTests.Build.cs b/Engine/Source/Programs/LowLevelTests/FoundationTests/FoundationTests.Build.cs
--- a/Engine/Source/Programs/LowLevelTests/FoundationTests/FoundationTests.Build.cs
+++ b/Engine/Source/Programs/LowLevelTests/FoundationTests/FoundationTests.Build.cs
@@ -51,6 +51,10 @@ public class FoundationTests : TestModuleRules
 				"Serialization",
 			});
 
+		PrivateIncludePathModuleNames.AddRange(new string[] { 
+			"NetCore"
+		});
+
 		if (Target.bBuildWithEditorOnlyData)
 		{
 			PrivateDependencyModuleNames.AddRange(
diff --git a/Engine/Source/Programs/LowLevelTests/FoundationTests/Private/TestObject.cpp b/Engine/Source/Programs/LowLevelTests/FoundationTests/Private/TestObject.cpp
--- a/Engine/Source/Programs/LowLevelTests/FoundationTests/Private/TestObject.cpp
+++ b/Engine/Source/Programs/LowLevelTests/FoundationTests/Private/TestObject.cpp
@@ -1,6 +1,9 @@
 // Copyright Epic Games, Inc. All Rights Reserved.
 
 #include "TestObject.h"
+#include "UObject/CoreNet.h"
+
+#include UE_INLINE_GENERATED_CPP_BY_NAME(TestObject)
 
 UTestInterface::UTestInterface(const FObjectInitializer& ObjectInitializer)
     : Super(ObjectInitializer)
@@ -20,4 +23,14 @@ int32 UTestObject::GetBPOverrideableValue_Implementation()
 FName UTestObject::GetNumberedName(FName BaseName, int32 InNumber)
 {
     return FName(BaseName, InNumber);
+}
+
+void UTestObject::OnRep_StrongObjectReference()
+{
+}
+
+void UTestObject::GetLifetimeReplicatedProps(TArray< FLifetimeProperty >& OutLifetimeProps) const
+{
+	Super::GetLifetimeReplicatedProps(OutLifetimeProps);
+    OutLifetimeProps.Emplace(UTestObject::StaticClass()->FindPropertyByName("StrongObjectReference")->RepIndex);
 }
\ No newline at end of file
diff --git a/Engine/Source/Programs/LowLevelTests/FoundationTests/Private/TestObject.h b/Engine/Source/Programs/LowLevelTests/FoundationTests/Private/TestObject.h
--- a/Engine/Source/Programs/LowLevelTests/FoundationTests/Private/TestObject.h
+++ b/Engine/Source/Programs/LowLevelTests/FoundationTests/Private/TestObject.h
@@ -1,5 +1,7 @@
 // Copyright Epic Games, Inc. All Rights Reserved.
 
+#pragma once
+
 #include "CoreTypes.h"
 #include "TestInterface.h"
 #include "UObject/Object.h"
@@ -34,7 +36,7 @@ class UTestObject : public UObject, public ITestInterface
     GENERATED_BODY()
 
 public:
-    UPROPERTY()
+    UPROPERTY(ReplicatedUsing=OnRep_StrongObjectReference)
     TObjectPtr<UObject> StrongObjectReference;
 
 #if WITH_EDITORONLY_DATA
@@ -80,4 +82,7 @@ public:
 
     UFUNCTION(BlueprintCallable, Category="Test")
     virtual FName GetNumberedName(FName BaseName, int32 InNumber) override;
+
+    UFUNCTION()
+    void OnRep_StrongObjectReference();
 };
\ No newline at end of file
diff --git a/Engine/Source/Programs/LowLevelTests/FoundationTests/Tests/UObjectTests.cpp b/Engine/Source/Programs/LowLevelTests/FoundationTests/Tests/UObjectTests.cpp
--- a/Engine/Source/Programs/LowLevelTests/FoundationTests/Tests/UObjectTests.cpp
+++ b/Engine/Source/Programs/LowLevelTests/FoundationTests/Tests/UObjectTests.cpp
@@ -2,11 +2,13 @@
 
 #include "CoreTypes.h"
 
+#include "TestCommon/Expectations.h"
 #include "TestHarness.h"
 #include "TestObject.h"
 #include "UObject/Class.h"
 #include "UObject/Package.h"
 
+
 TEST_CASE("CoreUObject::UFunction::Basic", "[CoreUObject]")
 {
     FName PackageName = MakeUniqueObjectName(nullptr, UPackage::StaticClass(), FName("/Memory/UFunctionBasic"));
@@ -16,4 +18,39 @@ TEST_CASE("CoreUObject::UFunction::Basic", "[CoreUObject]")
     int32 Value = Obj->GetBPOverrideableValue();
     int32 Value2 = Obj->GetBPOverrideableValue_Implementation();
     CHECK_EQUAL(Value, Value2);
+}
+
+TEST_CASE("CoreUObject::FProperty::Basic", "[CoreUObject]")
+{
+    FProperty* StrongObjectProp = UTestObject::StaticClass()->FindPropertyByName("StrongObjectReference");
+    REQUIRE(StrongObjectProp != nullptr);
+    CHECK_EQUAL(FName("OnRep_StrongObjectReference"), StrongObjectProp->RepNotifyFunc);
+}
+
+TEST_CASE("CoreUObject::Interface::Basic", "[CoreUObject]")
+{
+    UClass* TestObjectClass = UTestObject::StaticClass();
+    CHECK(TestObjectClass->ImplementsInterface(UTestInterface::StaticClass()));
+
+    FName PackageName = MakeUniqueObjectName(nullptr, UPackage::StaticClass(), FName("/Memory/InterfaceBasic"));
+    UPackage* Package = NewObject<UPackage>(nullptr, UPackage::StaticClass(), PackageName);
+    UTestObject* Obj = NewObject<UTestObject>(Package, "TestObject");
+
+    ITestInterface* InterfaceObj = Cast<ITestInterface>(Obj);
+    CHECK(InterfaceObj != nullptr);
+
+    struct FParams
+    {
+        FName BaseName = "TestName";
+        int32 Number = 17;
+        FName Result;
+    } Params;
+    CHECK_EQUALS("Native call", 
+        InterfaceObj->GetNumberedName(Params.BaseName, Params.Number),
+        FName(Params.BaseName, Params.Number));
+
+    UFunction* Function = UTestInterface::StaticClass()->FindFunctionByName("GetNumberedName");
+    CHECK(Function != nullptr);
+    Obj->ProcessEvent(Function, &Params);
+    CHECK_EQUALS("BP call", Params.Result, FName(Params.BaseName, Params.Number));
 }
\ No newline at end of file
