Namespaced UObject type test declarations.
Fix some generated code to use the fully qualified names of types to fix compile errors.

diff --git a/Engine/Source/Programs/LowLevelTests/FoundationTests/FoundationTests.Build.cs b/Engine/Source/Programs/LowLevelTests/FoundationTests/FoundationTests.Build.cs
--- a/Engine/Source/Programs/LowLevelTests/FoundationTests/FoundationTests.Build.cs
+++ b/Engine/Source/Programs/LowLevelTests/FoundationTests/FoundationTests.Build.cs
@@ -41,6 +41,8 @@ public class FoundationTests : TestModuleRules
 	}
 	public FoundationTests(ReadOnlyTargetRules Target) : base(Target, true)
 	{
+		bAllowUETypesInNamespaces = true;
+
 		PrivateDependencyModuleNames.AddRange(
 			new string[] {
 				"Core",
diff --git a/Engine/Source/Programs/LowLevelTests/FoundationTests/Private/NamespacedTypes.cpp b/Engine/Source/Programs/LowLevelTests/FoundationTests/Private/NamespacedTypes.cpp
new file mode 100644
--- /dev/null
+++ b/Engine/Source/Programs/LowLevelTests/FoundationTests/Private/NamespacedTypes.cpp
@@ -0,0 +1,13 @@
+// Copyright Epic Games, Inc. All Rights Reserved.
+
+#include "NamespacedTypes.h"
+
+int32 UE::Tests::UNamespacedObject::BPCallable(const FNamespacedStruct& StructParam, ENamespacedEnum EnumParam, UNamespacedObject* ObjectParam, TSubclassOf<UNamespacedObject> ClassParam)
+{
+    return StructParam.Value + Value;
+}
+
+int32 UNonNamespacedObject::BPCallable(const UE::Tests::FNamespacedStruct& StructParam, UE::Tests::ENamespacedEnum EnumParam, UE::Tests::UNamespacedObject* ObjectParam, TSubclassOf<UE::Tests::UNamespacedObject> ClassParam)
+{
+    return StructParam.Value;
+}
\ No newline at end of file
diff --git a/Engine/Source/Programs/LowLevelTests/FoundationTests/Private/NamespacedTypes.h b/Engine/Source/Programs/LowLevelTests/FoundationTests/Private/NamespacedTypes.h
new file mode 100644
--- /dev/null
+++ b/Engine/Source/Programs/LowLevelTests/FoundationTests/Private/NamespacedTypes.h
@@ -0,0 +1,128 @@
+// Copyright Epic Games, Inc. All Rights Reserved.
+
+#pragma once
+
+#include "CoreTypes.h"
+#include "Templates/SubclassOf.h"
+#include "UObject/Object.h"
+
+#include "NamespacedTypes.generated.h"
+
+namespace UE::Tests
+{
+
+UENUM(BlueprintType)
+enum class ENamespacedEnum : uint8
+{
+    One,
+    Two,
+    Three,
+};
+
+inline uint32 GetTypeHash(ENamespacedEnum Value)
+{
+    return ::GetTypeHash((uint8)Value);
+}
+
+USTRUCT(BlueprintType)
+struct FNamespacedStruct
+{
+    GENERATED_BODY()
+
+    UPROPERTY()
+    int32 Value;
+
+    UPROPERTY()
+    ENamespacedEnum Enum;
+
+    UPROPERTY()
+    TArray<ENamespacedEnum> ArrayOfEnums;
+};
+
+inline uint32 GetTypeHash(FNamespacedStruct Value)
+{
+    return ::GetTypeHash(Value.Value);
+}
+
+UCLASS()
+class UNamespacedObject : public UObject
+{
+    GENERATED_BODY()
+
+public:
+    UPROPERTY()
+    int32 Value;
+
+    UPROPERTY()
+    TObjectPtr<UNamespacedObject> Object;
+    UPROPERTY()
+    TArray<TObjectPtr<UNamespacedObject>> ArrayOfObjects;
+    UPROPERTY()
+    TSet<TObjectPtr<UNamespacedObject>> SetOfObjects;
+    UPROPERTY()
+    TMap<TObjectPtr<UNamespacedObject>, TObjectPtr<UNamespacedObject>> MapOfObjects;
+
+    UPROPERTY()
+    TSubclassOf<UNamespacedObject> Class;
+    UPROPERTY()
+    TArray<TSubclassOf<UNamespacedObject>> ArrayOfClasses;
+    UPROPERTY()
+    TSet<TSubclassOf<UNamespacedObject>> SetOfClasses;
+    UPROPERTY()
+    TMap<TSubclassOf<UNamespacedObject>, TSubclassOf<UNamespacedObject>> MapOfClasses;
+
+    UPROPERTY()
+    FNamespacedStruct Struct;
+    UPROPERTY()
+    TArray<FNamespacedStruct> ArrayOfStructs;
+    UPROPERTY()
+    TSet<FNamespacedStruct> SetOfStructs;
+    UPROPERTY()
+    TMap<FNamespacedStruct, FNamespacedStruct> MapOfStructs;
+
+    UPROPERTY()
+    ENamespacedEnum Enum;
+    UPROPERTY()
+    TArray<ENamespacedEnum> ArrayOfEnums;
+    UPROPERTY()
+    TSet<ENamespacedEnum> SetOfEnums;
+    UPROPERTY()
+    TMap<ENamespacedEnum, ENamespacedEnum> MapOfEnums;
+
+    UFUNCTION(BlueprintCallable, Category=Test)
+    int32 BPCallable(const FNamespacedStruct& StructParam, ENamespacedEnum EnumParam, UNamespacedObject* ObjectParam, TSubclassOf<UNamespacedObject> ClassParam);
+
+    UFUNCTION(BlueprintImplementableEvent, Category=Test)
+    void BPImplementable(const FNamespacedStruct& StructParam, ENamespacedEnum EnumParam, UNamespacedObject* ObjectParam, TSubclassOf<UNamespacedObject> ClassParam);
+};
+
+}
+
+// Class referencing types inside the namespace 
+UCLASS()
+class UNonNamespacedObject : public UObject
+{
+    GENERATED_BODY()
+
+public:
+    UPROPERTY()
+    TObjectPtr<UE::Tests::UNamespacedObject> Object;
+
+    UPROPERTY()
+    UE::Tests::FNamespacedStruct Struct;
+
+    UPROPERTY()
+    TArray<UE::Tests::FNamespacedStruct> ArrayOfStructs;
+
+    UPROPERTY()
+    UE::Tests::ENamespacedEnum Enum;
+
+    UPROPERTY()
+    TArray<UE::Tests::ENamespacedEnum> ArrayOfEnums;
+
+    UFUNCTION(BlueprintCallable, Category=Test)
+    int32 BPCallable(const UE::Tests::FNamespacedStruct& StructParam, UE::Tests::ENamespacedEnum EnumParam, UE::Tests::UNamespacedObject* ObjectParam, TSubclassOf<UE::Tests::UNamespacedObject> ClassParam);
+
+    UFUNCTION(BlueprintImplementableEvent, Category=Test)
+    void BPImplementable(const UE::Tests::FNamespacedStruct& StructParam, UE::Tests::ENamespacedEnum EnumParam, UE::Tests::UNamespacedObject* ObjectParam, TSubclassOf<UE::Tests::UNamespacedObject> ClassParam);
+};
diff --git a/Engine/Source/Programs/Shared/EpicGames.UHT/Exporters/CodeGen/UhtHeaderCodeGenerator.cs b/Engine/Source/Programs/Shared/EpicGames.UHT/Exporters/CodeGen/UhtHeaderCodeGenerator.cs
--- a/Engine/Source/Programs/Shared/EpicGames.UHT/Exporters/CodeGen/UhtHeaderCodeGenerator.cs
+++ b/Engine/Source/Programs/Shared/EpicGames.UHT/Exporters/CodeGen/UhtHeaderCodeGenerator.cs
@@ -225,7 +225,7 @@ namespace EpicGames.UHT.Exporters.CodeGen
 			builder.Append(' ');
 			if (!isDeclaration && outerClass != null)
 			{
-				builder.AppendClassSourceNameOrInterfaceName(outerClass).Append("::");
+				builder.Append(outerClass.Namespace.FullSourceName).AppendClassSourceNameOrInterfaceName(outerClass).Append("::");
 			}
 
 			if (alternateFunctionName != null)
diff --git a/Engine/Source/Programs/Shared/EpicGames.UHT/Exporters/CodeGen/UhtHeaderCodeGeneratorCppFile.cs b/Engine/Source/Programs/Shared/EpicGames.UHT/Exporters/CodeGen/UhtHeaderCodeGeneratorCppFile.cs
--- a/Engine/Source/Programs/Shared/EpicGames.UHT/Exporters/CodeGen/UhtHeaderCodeGeneratorCppFile.cs
+++ b/Engine/Source/Programs/Shared/EpicGames.UHT/Exporters/CodeGen/UhtHeaderCodeGeneratorCppFile.cs
@@ -574,21 +574,19 @@ namespace EpicGames.UHT.Exporters.CodeGen
 				if (scriptStruct.SuperScriptStruct != null)
 				{
 					builder.Append("static_assert(std::is_polymorphic<")
-						.Append(scriptStruct.Namespace.FullSourceName)
-						.Append(scriptStruct.SourceName)
+						.Append(scriptStruct.FullyQualifiedSourceName)
 						.Append(">() == std::is_polymorphic<")
-						.Append(scriptStruct.SuperScriptStruct.Namespace.FullSourceName)
-						.Append(scriptStruct.SuperScriptStruct.SourceName)
+						.Append(scriptStruct.SuperScriptStruct.FullyQualifiedSourceName)
 						.Append(">(), \"USTRUCT ")
-						.Append(scriptStruct.SourceName)
+						.Append(scriptStruct.FullyQualifiedSourceName)
 						.Append(" cannot be polymorphic unless super ")
-						.Append(scriptStruct.SuperScriptStruct.SourceName)
+						.Append(scriptStruct.SuperScriptStruct.FullyQualifiedSourceName)
 						.Append(" is polymorphic\");\r\n");
 				}
 
 				// Outer singleton
 				builder.Append("static FStructRegistrationInfo ").Append(registrationName).Append(";\r\n");
-				builder.Append("class UScriptStruct* ").Append(scriptStruct.Namespace.FullSourceName).Append(scriptStruct.SourceName).Append("::StaticStruct()\r\n");
+				builder.Append($"class UScriptStruct* {scriptStruct.FullyQualifiedSourceName}::StaticStruct()\r\n");
 				builder.Append("{\r\n");
 				builder.Append("\tif (!").Append(registrationName).Append(".OuterSingleton)\r\n");
 				builder.Append("\t{\r\n");
diff --git a/Engine/Source/Programs/Shared/EpicGames.UHT/Types/Properties/UhtEnumProperty.cs b/Engine/Source/Programs/Shared/EpicGames.UHT/Types/Properties/UhtEnumProperty.cs
--- a/Engine/Source/Programs/Shared/EpicGames.UHT/Types/Properties/UhtEnumProperty.cs
+++ b/Engine/Source/Programs/Shared/EpicGames.UHT/Types/Properties/UhtEnumProperty.cs
@@ -109,26 +109,26 @@ namespace EpicGames.UHT.Types
 						(!isTemplateArgument && (property.PropertyFlags.HasAnyFlags(EPropertyFlags.ReturnParm) || 
 						!property.PropertyFlags.HasAnyFlags(EPropertyFlags.OutParm))))
 					{
-						builder.Append(enumObj.CppType);
+						builder.Append(enumObj.FullyQualifiedCppType);
 					}
 					else
 					{
-						builder.Append("TEnumAsByte<").Append(enumObj.CppType).Append('>');
+						builder.Append("TEnumAsByte<").Append(enumObj.FullyQualifiedCppType).Append('>');
 					}
 					break;
 
 				case UhtPropertyTextType.FunctionThunkParameterArgType:
-					builder.Append(enumObj.SourceName);
+					builder.Append(enumObj.FullyQualifiedSourceName);
 					break;
 
 				case UhtPropertyTextType.RigVMType:
 					if (enumObj.CppForm == UhtEnumCppForm.EnumClass || !isTemplateArgument)
 					{
-						builder.Append(enumObj.CppType);
+						builder.Append(enumObj.FullyQualifiedCppType);
 					}
 					else
 					{
-						builder.Append("TEnumAsByte<").Append(enumObj.CppType).Append('>');
+						builder.Append("TEnumAsByte<").Append(enumObj.FullyQualifiedCppType).Append('>');
 					}
 					break;
 
@@ -143,11 +143,11 @@ namespace EpicGames.UHT.Types
 				default:
 					if (enumObj.CppForm == UhtEnumCppForm.EnumClass)
 					{
-						builder.Append(enumObj.CppType);
+						builder.Append(enumObj.FullyQualifiedCppType);
 					}
 					else
 					{
-						builder.Append("TEnumAsByte<").Append(enumObj.CppType).Append('>');
+						builder.Append("TEnumAsByte<").Append(enumObj.FullyQualifiedCppType).Append('>');
 					}
 					break;
 			}
@@ -165,15 +165,15 @@ namespace EpicGames.UHT.Types
 		{
 			if (!property.PropertyFlags.HasAnyFlags(EPropertyFlags.OutParm))
 			{
-				builder.Append(enumObj.CppType).Append('(').AppendFunctionThunkParameterName(property).Append(')');
+				builder.Append(enumObj.FullyQualifiedCppType).Append('(').AppendFunctionThunkParameterName(property).Append(')');
 			}
 			else if (enumObj.CppForm == UhtEnumCppForm.EnumClass)
 			{
-				builder.Append('(').Append(enumObj.CppType).Append("&)(").AppendFunctionThunkParameterName(property).Append(')');
+				builder.Append($"({enumObj.FullyQualifiedCppType}&)(").AppendFunctionThunkParameterName(property).Append(')');
 			}
 			else
 			{
-				builder.Append("(TEnumAsByte<").Append(enumObj.CppType).Append(">&)(").AppendFunctionThunkParameterName(property).Append(')');
+				builder.Append($"(TEnumAsByte<{enumObj.FullyQualifiedCppType}>&)(").AppendFunctionThunkParameterName(property).Append(')');
 			}
 			return builder;
 		}
diff --git a/Engine/Source/Programs/Shared/EpicGames.UHT/Types/Properties/UhtSetProperty.cs b/Engine/Source/Programs/Shared/EpicGames.UHT/Types/Properties/UhtSetProperty.cs
--- a/Engine/Source/Programs/Shared/EpicGames.UHT/Types/Properties/UhtSetProperty.cs
+++ b/Engine/Source/Programs/Shared/EpicGames.UHT/Types/Properties/UhtSetProperty.cs
@@ -132,9 +132,9 @@ namespace EpicGames.UHT.Types
 				builder
 					.AppendTabs(tabs)
 					.Append("static_assert(TModels_V<CGetTypeHashable, ")
-					.Append(structProperty.ScriptStruct.SourceName)
+					.Append(structProperty.ScriptStruct.FullyQualifiedSourceName)
 					.Append(">, \"The structure '")
-					.Append(structProperty.ScriptStruct.SourceName)
+					.Append(structProperty.ScriptStruct.FullyQualifiedSourceName)
 					.Append("' is used in a TSet but does not have a GetValueTypeHash defined\");\r\n");
 			}
 
diff --git a/Engine/Source/Programs/Shared/EpicGames.UHT/Types/Properties/UhtStructProperty.cs b/Engine/Source/Programs/Shared/EpicGames.UHT/Types/Properties/UhtStructProperty.cs
--- a/Engine/Source/Programs/Shared/EpicGames.UHT/Types/Properties/UhtStructProperty.cs
+++ b/Engine/Source/Programs/Shared/EpicGames.UHT/Types/Properties/UhtStructProperty.cs
@@ -124,7 +124,7 @@ namespace EpicGames.UHT.Types
 					break;
 
 				default:
-					builder.Append(ScriptStruct.SourceName);
+					builder.Append(ScriptStruct.FullyQualifiedSourceName);
 					break;
 			}
 			return builder;
diff --git a/Engine/Source/Programs/Shared/EpicGames.UHT/Types/UhtEnum.cs b/Engine/Source/Programs/Shared/EpicGames.UHT/Types/UhtEnum.cs
--- a/Engine/Source/Programs/Shared/EpicGames.UHT/Types/UhtEnum.cs
+++ b/Engine/Source/Programs/Shared/EpicGames.UHT/Types/UhtEnum.cs
@@ -200,6 +200,12 @@ namespace EpicGames.UHT.Types
 		/// </summary>
 		public string CppType { get; set; } = String.Empty;
 
+		/// <summary>
+		/// Fully qualified cpp type name. Inc;ude enclosing namespace as well as CppType (which may include enum's own namespace) 
+		/// </summary>
+		public string FullyQualifiedCppType => _fullyQualifiedCppType ??= $"{Namespace.FullSourceName}{CppType}";
+		private string? _fullyQualifiedCppType = null;
+
 		/// <summary>
 		/// Collection of enumeration values
 		/// </summary>
diff --git a/Engine/Source/Programs/Shared/EpicGames.UHT/Types/UhtField.cs b/Engine/Source/Programs/Shared/EpicGames.UHT/Types/UhtField.cs
--- a/Engine/Source/Programs/Shared/EpicGames.UHT/Types/UhtField.cs
+++ b/Engine/Source/Programs/Shared/EpicGames.UHT/Types/UhtField.cs
@@ -125,6 +125,12 @@ namespace EpicGames.UHT.Types
 		/// </summary>
 		public UhtNamespace Namespace { get; set; }
 
+		/// <summary>
+		/// Fully qualified name of the field in source, i.e. Namespace and SourceName concatenated.
+		/// </summary>
+		public string FullyQualifiedSourceName => _fullyQualifiedSourceName ??= $"{Namespace.FullSourceName}{SourceName}";
+		private string? _fullyQualifiedSourceName = null;
+
 		/// <summary>
 		/// Construct a new field
 		/// </summary>
