Fix platform issues Foundation Tests

diff --git a/Engine/Source/Programs/LowLevelTests/FoundationTests/FoundationTests.Build.cs b/Engine/Source/Programs/LowLevelTests/FoundationTests/FoundationTests.Build.cs
--- a/Engine/Source/Programs/LowLevelTests/FoundationTests/FoundationTests.Build.cs
+++ b/Engine/Source/Programs/LowLevelTests/FoundationTests/FoundationTests.Build.cs
@@ -66,6 +66,11 @@ public class FoundationTests : TestModuleRules
 				});
 		}
 
+		if (Target.bCompileAgainstEngine)
+		{
+			PrivateIncludePathModuleNames.AddRange(new string[]{ "SlateCore" });
+		}
+
 		if (Target.bBuildWithEditorOnlyData &&
 			(Target.Platform == UnrealTargetPlatform.Mac ||
 			 Target.Platform == UnrealTargetPlatform.Win64 ||
diff --git a/Engine/Source/Programs/LowLevelTests/FoundationTests/Tests/TestGroupEvents.cpp b/Engine/Source/Programs/LowLevelTests/FoundationTests/Tests/TestGroupEvents.cpp
--- a/Engine/Source/Programs/LowLevelTests/FoundationTests/Tests/TestGroupEvents.cpp
+++ b/Engine/Source/Programs/LowLevelTests/FoundationTests/Tests/TestGroupEvents.cpp
@@ -1,10 +1,23 @@
 // Copyright Epic Games, Inc. All Rights Reserved.
 #include "TestCommon/Initialization.h"
 
+#if WITH_ENGINE
+#include "Logging/LogScopedVerbosityOverride.h"
+#include "SlateGlobals.h"
+#include "Styling/SlateWidgetStyleContainerBase.h"
+#endif
+
 #include <catch2/catch_test_macros.hpp>
 
 GROUP_BEFORE_GLOBAL(Catch::DefaultGroup)
 {
+#if WITH_ENGINE
+	LOG_SCOPE_VERBOSITY_OVERRIDE(LogSlate, ELogVerbosity::Error);
+	LOG_SCOPE_VERBOSITY_OVERRIDE(LogSlateStyle, ELogVerbosity::Error);
+	LOG_SCOPE_VERBOSITY_OVERRIDE(LogUObjectGlobals, ELogVerbosity::Fatal);
+	LOG_SCOPE_VERBOSITY_OVERRIDE(LogStreaming, ELogVerbosity::Error);
+#endif
+
 	InitAll(true, true);
 }
 
