Fixed two issues with UHT and types in namespaces.

1) If in the same source file, two namespaces appeared with a common parent namespace, UHT would generate an internal error.
2) Interfaces in namespaces would fail compiling the VTABLE_OFFSET macro.

#rnx

diff --git a/Engine/Source/Programs/LowLevelTests/FoundationTests/Private/NamespacedTypes.h b/Engine/Source/Programs/LowLevelTests/FoundationTests/Private/NamespacedTypes.h
--- a/Engine/Source/Programs/LowLevelTests/FoundationTests/Private/NamespacedTypes.h
+++ b/Engine/Source/Programs/LowLevelTests/FoundationTests/Private/NamespacedTypes.h
@@ -4,6 +4,7 @@
 
 #include "CoreTypes.h"
 #include "Templates/SubclassOf.h"
+#include "UObject/Interface.h"
 #include "UObject/Object.h"
 
 #include "NamespacedTypes.generated.h"
@@ -106,6 +107,40 @@ private:
 
 }
 
+namespace UE::DifferentNameSpaceWithSameParent
+{
+USTRUCT()
+struct FDifferentNameSpaceStructTest
+{
+	GENERATED_BODY()
+	UPROPERTY()
+	int Bob;
+};
+} // namespace UE::DifferentNameSpaceWithSameParent
+
+namespace UE::InterfaceTest
+{
+UINTERFACE()
+class UInterfaceTest : public UInterface
+{
+	GENERATED_BODY()
+};
+
+class IInterfaceTest
+{
+	GENERATED_BODY()
+};
+} // namespace UE::InterfaceTest
+
+namespace UE::InterfaceTestUsage
+{
+UCLASS()
+class UInterfaceTestUsage : public UObject, public UE::InterfaceTest::IInterfaceTest
+{
+	GENERATED_BODY()
+};
+} // namespace UE::InterfaceTestUsage
+
 // Class referencing types inside the namespace 
 UCLASS()
 class UNonNamespacedObject : public UObject
diff --git a/Engine/Source/Programs/Shared/EpicGames.UHT/Exporters/CodeGen/UhtHeaderCodeGeneratorCppFile.cs b/Engine/Source/Programs/Shared/EpicGames.UHT/Exporters/CodeGen/UhtHeaderCodeGeneratorCppFile.cs
--- a/Engine/Source/Programs/Shared/EpicGames.UHT/Exporters/CodeGen/UhtHeaderCodeGeneratorCppFile.cs
+++ b/Engine/Source/Programs/Shared/EpicGames.UHT/Exporters/CodeGen/UhtHeaderCodeGeneratorCppFile.cs
@@ -2541,9 +2541,9 @@ namespace EpicGames.UHT.Exporters.CodeGen
 				{
 					builder
 						.Append(", (int32)VTABLE_OFFSET(")
-						.Append(classObj.SourceName)
+						.Append(classObj.FullyQualifiedSourceName)
 						.Append(", ")
-						.Append(interfaceObj.SourceName)
+						.Append(interfaceObj.FullyQualifiedSourceName)
 						.Append("), false");
 				}
 				builder
diff --git a/Engine/Source/Programs/Shared/EpicGames.UHT/Utils/UhtSession.cs b/Engine/Source/Programs/Shared/EpicGames.UHT/Utils/UhtSession.cs
--- a/Engine/Source/Programs/Shared/EpicGames.UHT/Utils/UhtSession.cs
+++ b/Engine/Source/Programs/Shared/EpicGames.UHT/Utils/UhtSession.cs
@@ -1747,7 +1747,7 @@ namespace EpicGames.UHT.Utils
 					return new(namespaceElements[index - 1].Namespace, namespaceElements[index].Name.ToString(), fullName, false);
 				});
 			}
-			headerNamespaceLookupCache.Add(fullName, namespaceObj);
+			headerNamespaceLookupCache.TryAdd(fullName, namespaceObj);
 			namespaceElements[index] = new(namespaceElements[index].Name, namespaceObj);
 			return namespaceObj;
 		}
