[UObjectAnnotationTests]
- Increase the timeout to avoid CIS failures on slower platforms

diff --git a/Engine/Source/Programs/LowLevelTests/FoundationTests/Tests/UObjectAnnotationTests.cpp b/Engine/Source/Programs/LowLevelTests/FoundationTests/Tests/UObjectAnnotationTests.cpp
--- a/Engine/Source/Programs/LowLevelTests/FoundationTests/Tests/UObjectAnnotationTests.cpp
+++ b/Engine/Source/Programs/LowLevelTests/FoundationTests/Tests/UObjectAnnotationTests.cpp
@@ -26,7 +26,7 @@ TEST_CASE("UObjectAnnotation::ThreadSafety", "[UObjectAnnotation]")
 				double StartTime = FPlatformTime::Seconds();
 				while (!DoneEvent.WaitFor(UE::FMonotonicTimeSpan::FromSeconds(1)))
 				{
-					if (FPlatformTime::Seconds() - StartTime > 10.0f)
+					if (FPlatformTime::Seconds() - StartTime > 60.0f)
 					{
 						UE_LOG(LogTemp, Fatal, TEXT("UObjectAnnotation::ThreadSafety is non responsive and likely dead-locked. Aborting program."));
 					}
