Added support for patterns such as "FString()" for default values in UFUNCTIONs.

diff --git a/Engine/Source/Programs/LowLevelTests/FoundationTests/Private/NamespacedTypes.h b/Engine/Source/Programs/LowLevelTests/FoundationTests/Private/NamespacedTypes.h
--- a/Engine/Source/Programs/LowLevelTests/FoundationTests/Private/NamespacedTypes.h
+++ b/Engine/Source/Programs/LowLevelTests/FoundationTests/Private/NamespacedTypes.h
@@ -119,6 +119,72 @@ public:
     UFUNCTION(BlueprintImplementableEvent, Category=Test)
     void BPImplementable(const FNamespacedStruct& StructParam, ENamespacedEnum EnumParam, UNamespacedObject* ObjectParam, TSubclassOf<UNamespacedObject> ClassParam);
 
+
+
+	// Test default values with UFUNCTION parsing
+	UFUNCTION(BlueprintCallable, Category = "Test")
+	void TestFString1(FString String = FString(TEXT(""))) {}
+
+	UFUNCTION(BlueprintCallable, Category = "Test")
+	void TestFString2(FString String = FString("")) {}
+
+	UFUNCTION(BlueprintCallable, Category = "Test")
+	void TestFString3(FString String = FString()) {}
+
+	UFUNCTION(BlueprintCallable, Category = "Test")
+	void TestFString4(FString String = TEXT("")) {}
+
+	UFUNCTION(BlueprintCallable, Category = "Test")
+	void TestFString5(FString String = "") {}
+
+#if 0
+	// These aren't allowed for callable functions.
+	UFUNCTION(BlueprintCallable, Category = "Test")
+	void TestUtf8String1(FUtf8String String = FUtf8String(TEXT(""))) {}
+
+	UFUNCTION(BlueprintCallable, Category = "Test")
+	void TestUtf8String2(FUtf8String String = FUtf8String("")) {}
+
+	UFUNCTION(BlueprintCallable, Category = "Test")
+	void TestUtf8String3(FUtf8String String = FUtf8String()) {}
+
+	UFUNCTION(BlueprintCallable, Category = "Test")
+	void TestUtf8String4(FUtf8String String = TEXT("")) {}
+
+	UFUNCTION(BlueprintCallable, Category = "Test")
+	void TestUtf8String5(FUtf8String String = "") {}
+
+	UFUNCTION(BlueprintCallable, Category = "Test")
+	void TestAnsiString1(FAnsiString String = FAnsiString(TEXT(""))) {}
+
+	UFUNCTION(BlueprintCallable, Category = "Test")
+	void TestAnsiString2(FAnsiString String = FAnsiString("")) {}
+
+	UFUNCTION(BlueprintCallable, Category = "Test")
+	void TestAnsiString3(FAnsiString String = FAnsiString()) {}
+
+	UFUNCTION(BlueprintCallable, Category = "Test")
+	void TestAnsiString4(FAnsiString String = TEXT("")) {}
+
+	UFUNCTION(BlueprintCallable, Category = "Test")
+	void TestAnsiString5(FAnsiString String = "") {}
+#endif
+
+	UFUNCTION(BlueprintCallable, Category = "Test")
+	void TestFNameString1(FName String = FName(TEXT(""))) {}
+
+	UFUNCTION(BlueprintCallable, Category = "Test")
+	void TestFNameString2(FName String = FName("")) {}
+
+	UFUNCTION(BlueprintCallable, Category = "Test")
+	void TestFNameString3(FName String = FName()) {}
+
+	UFUNCTION(BlueprintCallable, Category = "Test")
+	void TestFNameString4(FName String = TEXT("")) {}
+
+	UFUNCTION(BlueprintCallable, Category = "Test")
+	void TestFNameString5(FName String = "") {}
+
 private:
 	UPROPERTY()
 	int64 SomePrivateProperty;
diff --git a/Engine/Source/Programs/Shared/EpicGames.UHT/Tokenizer/UhtTokenReaderStringExtensions.cs b/Engine/Source/Programs/Shared/EpicGames.UHT/Tokenizer/UhtTokenReaderStringExtensions.cs
--- a/Engine/Source/Programs/Shared/EpicGames.UHT/Tokenizer/UhtTokenReaderStringExtensions.cs
+++ b/Engine/Source/Programs/Shared/EpicGames.UHT/Tokenizer/UhtTokenReaderStringExtensions.cs
@@ -121,5 +121,31 @@ namespace EpicGames.UHT.Tokenizer
 				return tokenReader.GetConstString(exceptionContext);
 			}
 		}
+
+		/// <summary>
+		/// Get a typed const string that can optionally be wrapped with a TEXT() macro
+		/// </summary>
+		/// <param name="tokenReader">Token reader</param>
+		/// <param name="typeName">Expected type name</param>
+		/// <param name="exceptionContext">If not null, an exception will be thrown with the given text as part of the message.</param>
+		/// <returns>The value of the string</returns>
+		public static StringView GetTypedConstString(this IUhtTokenReader tokenReader, string typeName, object? exceptionContext = null)
+		{
+			StringView output = new();
+			if (tokenReader.TryOptional(typeName))
+			{
+				tokenReader.Require('(');
+				if (!tokenReader.PeekToken().IsSymbol(')'))
+				{
+					output = tokenReader.GetWrappedConstString();
+				}
+				tokenReader.Require(')');
+			}
+			else
+			{
+				output = tokenReader.GetWrappedConstString();
+			}
+			return output;
+		}
 	}
 }
diff --git a/Engine/Source/Programs/Shared/EpicGames.UHT/Types/Properties/UhtAnsiStrProperty.cs b/Engine/Source/Programs/Shared/EpicGames.UHT/Types/Properties/UhtAnsiStrProperty.cs
--- a/Engine/Source/Programs/Shared/EpicGames.UHT/Types/Properties/UhtAnsiStrProperty.cs
+++ b/Engine/Source/Programs/Shared/EpicGames.UHT/Types/Properties/UhtAnsiStrProperty.cs
@@ -73,18 +73,8 @@ namespace EpicGames.UHT.Types
 		/// <inheritdoc/>
 		public override bool SanitizeDefaultValue(IUhtTokenReader defaultValueReader, StringBuilder innerDefaultValue)
 		{
-			if (defaultValueReader.TryOptional("FAnsiString"))
-			{
-				defaultValueReader.Require('(');
-				StringView value = defaultValueReader.GetWrappedConstString();
-				defaultValueReader.Require(')');
-				innerDefaultValue.Append(value);
-			}
-			else
-			{
-				StringView value = defaultValueReader.GetWrappedConstString();
-				innerDefaultValue.Append(value);
-			}
+			StringView value = defaultValueReader.GetTypedConstString("FAnsiString");
+			innerDefaultValue.Append(value);
 			return true;
 		}
 
diff --git a/Engine/Source/Programs/Shared/EpicGames.UHT/Types/Properties/UhtNameProperty.cs b/Engine/Source/Programs/Shared/EpicGames.UHT/Types/Properties/UhtNameProperty.cs
--- a/Engine/Source/Programs/Shared/EpicGames.UHT/Types/Properties/UhtNameProperty.cs
+++ b/Engine/Source/Programs/Shared/EpicGames.UHT/Types/Properties/UhtNameProperty.cs
@@ -76,16 +76,9 @@ namespace EpicGames.UHT.Types
 			{
 				innerDefaultValue.Append("None");
 			}
-			else if (defaultValueReader.TryOptional("FName"))
-			{
-				defaultValueReader.Require('(');
-				StringView value = defaultValueReader.GetWrappedConstString();
-				defaultValueReader.Require(')');
-				innerDefaultValue.Append(value);
-			}
 			else
 			{
-				StringView value = defaultValueReader.GetWrappedConstString();
+				StringView value = defaultValueReader.GetTypedConstString("FName");
 				innerDefaultValue.Append(value);
 			}
 			return true;
diff --git a/Engine/Source/Programs/Shared/EpicGames.UHT/Types/Properties/UhtStrProperty.cs b/Engine/Source/Programs/Shared/EpicGames.UHT/Types/Properties/UhtStrProperty.cs
--- a/Engine/Source/Programs/Shared/EpicGames.UHT/Types/Properties/UhtStrProperty.cs
+++ b/Engine/Source/Programs/Shared/EpicGames.UHT/Types/Properties/UhtStrProperty.cs
@@ -72,18 +72,8 @@ namespace EpicGames.UHT.Types
 		/// <inheritdoc/>
 		public override bool SanitizeDefaultValue(IUhtTokenReader defaultValueReader, StringBuilder innerDefaultValue)
 		{
-			if (defaultValueReader.TryOptional("FString"))
-			{
-				defaultValueReader.Require('(');
-				StringView value = defaultValueReader.GetWrappedConstString();
-				defaultValueReader.Require(')');
-				innerDefaultValue.Append(value);
-			}
-			else
-			{
-				StringView value = defaultValueReader.GetWrappedConstString();
-				innerDefaultValue.Append(value);
-			}
+			StringView value = defaultValueReader.GetTypedConstString("FString");
+			innerDefaultValue.Append(value);
 			return true;
 		}
 
diff --git a/Engine/Source/Programs/Shared/EpicGames.UHT/Types/Properties/UhtUtf8StrProperty.cs b/Engine/Source/Programs/Shared/EpicGames.UHT/Types/Properties/UhtUtf8StrProperty.cs
--- a/Engine/Source/Programs/Shared/EpicGames.UHT/Types/Properties/UhtUtf8StrProperty.cs
+++ b/Engine/Source/Programs/Shared/EpicGames.UHT/Types/Properties/UhtUtf8StrProperty.cs
@@ -73,18 +73,8 @@ namespace EpicGames.UHT.Types
 		/// <inheritdoc/>
 		public override bool SanitizeDefaultValue(IUhtTokenReader defaultValueReader, StringBuilder innerDefaultValue)
 		{
-			if (defaultValueReader.TryOptional("FUtf8String"))
-			{
-				defaultValueReader.Require('(');
-				StringView value = defaultValueReader.GetWrappedConstString();
-				defaultValueReader.Require(')');
-				innerDefaultValue.Append(value);
-			}
-			else
-			{
-				StringView value = defaultValueReader.GetWrappedConstString();
-				innerDefaultValue.Append(value);
-			}
+			StringView value = defaultValueReader.GetTypedConstString("FUtf8String");
+			innerDefaultValue.Append(value);
 			return true;
 		}
 
