FPackageName: SplitPackageNameRoot: Fix unintentional change in behavior from CL 39458197 on edge case with no second slash: /root with default flags should go to mountpoint == root rather than mountpoint == /root.
#rnx

diff --git a/Engine/Source/Programs/LowLevelTests/FoundationTests/Tests/PackageNameTests.cpp b/Engine/Source/Programs/LowLevelTests/FoundationTests/Tests/PackageNameTests.cpp
--- a/Engine/Source/Programs/LowLevelTests/FoundationTests/Tests/PackageNameTests.cpp
+++ b/Engine/Source/Programs/LowLevelTests/FoundationTests/Tests/PackageNameTests.cpp
@@ -291,12 +291,12 @@ TEST_CASE("Split PackageNameRoot", "[CoreUObject][FPackageName::SplitPackageName
 
 	PackageName = TEXTVIEW("/root");
 	Root = FPackageName::SplitPackageNameRoot(PackageName, &RelPath);
-	REQUIRE(Root == TEXTVIEW("/root"));
+	REQUIRE(Root == TEXTVIEW("root"));
 	REQUIRE(RelPath == TEXTVIEW(""));
 
 	PackageName = TEXTVIEW("/");
 	Root = FPackageName::SplitPackageNameRoot(PackageName, &RelPath);
-	REQUIRE(Root == TEXTVIEW("/"));
+	REQUIRE(Root == TEXTVIEW(""));
 	REQUIRE(RelPath == TEXTVIEW(""));
 
 	PackageName = TEXTVIEW("");
@@ -334,6 +334,42 @@ TEST_CASE("Split PackageNameRoot", "[CoreUObject][FPackageName::SplitPackageName
 	REQUIRE(Root == TEXTVIEW("/root/"));
 	REQUIRE(RelPath == TEXTVIEW("path/leaf"));
 
+	// Path formats in edge case with no second slash 
+	PackageName = TEXTVIEW("/root");
+	Root = FPackageName::SplitPackageNameRoot(PackageName, &RelPath, FPackageName::EPathFormatFlags::MountPointNoSlashes);
+	REQUIRE(Root == TEXTVIEW("root"));
+	REQUIRE(RelPath == TEXTVIEW(""));
+
+	Root = FPackageName::SplitPackageNameRoot(PackageName, &RelPath, FPackageName::EPathFormatFlags::MountPointLeadingSlash);
+	REQUIRE(Root == TEXTVIEW("/root"));
+	REQUIRE(RelPath == TEXTVIEW(""));
+
+	Root = FPackageName::SplitPackageNameRoot(PackageName, &RelPath, FPackageName::EPathFormatFlags::MountPointTrailingSlash);
+	REQUIRE(Root == TEXTVIEW("root"));
+	REQUIRE(RelPath == TEXTVIEW(""));
+
+	Root = FPackageName::SplitPackageNameRoot(PackageName, &RelPath, FPackageName::EPathFormatFlags::MountPointSlashes);
+	REQUIRE(Root == TEXTVIEW("/root"));
+	REQUIRE(RelPath == TEXTVIEW(""));
+
+	// Path formats in edge case with a string equal to single slash
+	PackageName = TEXTVIEW("/");
+	Root = FPackageName::SplitPackageNameRoot(PackageName, &RelPath, FPackageName::EPathFormatFlags::MountPointNoSlashes);
+	REQUIRE(Root == TEXTVIEW(""));
+	REQUIRE(RelPath == TEXTVIEW(""));
+
+	Root = FPackageName::SplitPackageNameRoot(PackageName, &RelPath, FPackageName::EPathFormatFlags::MountPointLeadingSlash);
+	REQUIRE(Root == TEXTVIEW("/"));
+	REQUIRE(RelPath == TEXTVIEW(""));
+
+	Root = FPackageName::SplitPackageNameRoot(PackageName, &RelPath, FPackageName::EPathFormatFlags::MountPointTrailingSlash);
+	REQUIRE(Root == TEXTVIEW(""));
+	REQUIRE(RelPath == TEXTVIEW(""));
+
+	Root = FPackageName::SplitPackageNameRoot(PackageName, &RelPath, FPackageName::EPathFormatFlags::MountPointSlashes);
+	REQUIRE(Root == TEXTVIEW("/"));
+	REQUIRE(RelPath == TEXTVIEW(""));
+
 	// Verify Path formats are respected in the FName version
 	FString RootStr;
 	FString RelPathStr;
diff --git a/Engine/Source/Runtime/CoreUObject/Private/Misc/PackageName.cpp b/Engine/Source/Runtime/CoreUObject/Private/Misc/PackageName.cpp
--- a/Engine/Source/Runtime/CoreUObject/Private/Misc/PackageName.cpp
+++ b/Engine/Source/Runtime/CoreUObject/Private/Misc/PackageName.cpp
@@ -3027,7 +3027,14 @@ FStringView FPackageName::SplitPackageNameRoot(FStringView InPackageName, FStrin
 		{
 			*OutRelativePath = FStringView();
 		}
-		return InPackageName;
+		FStringView PackageRoot = InPackageName;
+		// Convert to the requested output format
+		if (!EnumHasAnyFlags(Flags, EPathFormatFlags::MountPointLeadingSlash))
+		{
+			PackageRoot.RightChopInline(1);
+		}
+
+		return PackageRoot;
 	}
 	check(SecondSlashIndex != INDEX_NONE);
 	// Change back to an index relative to the beginnning of InPackageName rather than relative to InPackageName[1].
