FoundationTests compilation fix
- ZenOplogUtils tests only run on desktop platforms since the library relies on BuildServerInterface

diff --git a/Engine/Source/Programs/LowLevelTests/FoundationTests/FoundationTests.Build.cs b/Engine/Source/Programs/LowLevelTests/FoundationTests/FoundationTests.Build.cs
--- a/Engine/Source/Programs/LowLevelTests/FoundationTests/FoundationTests.Build.cs
+++ b/Engine/Source/Programs/LowLevelTests/FoundationTests/FoundationTests.Build.cs
@@ -50,10 +50,15 @@ public class FoundationTests : TestModuleRules
 				"CoreUObject",
 				"TelemetryUtils",
 				"AssetRegistry",
-				"Serialization",
-				"ZenOplogUtils"
+				"Serialization"
 			});
 
+		// ZenOplogUtils relies on BuildServiceInterface which only runs on desktop platforms
+		if (Target.Platform.IsInGroup(UnrealPlatformGroup.Desktop))
+		{
+			PrivateDependencyModuleNames.Add("ZenOplogUtils");
+		}
+
 		PrivateIncludePathModuleNames.AddRange(new string[] { 
 			"NetCore"
 		});
