Comment fixup (BuildServiceInterface -> BuildServerInterface)

diff --git a/Engine/Source/Programs/LowLevelTests/FoundationTests/FoundationTests.Build.cs b/Engine/Source/Programs/LowLevelTests/FoundationTests/FoundationTests.Build.cs
--- a/Engine/Source/Programs/LowLevelTests/FoundationTests/FoundationTests.Build.cs
+++ b/Engine/Source/Programs/LowLevelTests/FoundationTests/FoundationTests.Build.cs
@@ -53,7 +53,7 @@ public class FoundationTests : TestModuleRules
 				"Serialization"
 			});
 
-		// ZenOplogUtils relies on BuildServiceInterface which only runs on desktop platforms
+		// ZenOplogUtils relies on BuildServerInterface which only runs on desktop platforms
 		if (Target.Platform.IsInGroup(UnrealPlatformGroup.Desktop))
 		{
 			PrivateDependencyModuleNames.Add("ZenOplogUtils");
