FoundationTests: Disable editor-only data for non-editor platforms

#Synced-CL 32500503

diff --git a/Engine/Source/Programs/LowLevelTests/FoundationTests/FoundationTests.Target.cs b/Engine/Source/Programs/LowLevelTests/FoundationTests/FoundationTests.Target.cs
--- a/Engine/Source/Programs/LowLevelTests/FoundationTests/FoundationTests.Target.cs
+++ b/Engine/Source/Programs/LowLevelTests/FoundationTests/FoundationTests.Target.cs
@@ -1,5 +1,6 @@
 // Copyright Epic Games, Inc. All Rights Reserved.
 
+using System.Linq;
 using UnrealBuildTool;
 
 [SupportedPlatforms(UnrealPlatformClass.All)]
@@ -15,7 +16,8 @@ public class FoundationTestsTarget : TestTargetRules
 
 		bCompileAgainstApplicationCore = true;
 
-		bBuildWithEditorOnlyData = Target.Platform.IsInGroup(UnrealPlatformGroup.Desktop)
+		bool bIsEditorPlatform = UnrealBuildTool.Utils.GetPlatformsInClass(UnrealPlatformClass.Editor).Contains(Target.Platform);
+		bBuildWithEditorOnlyData = bIsEditorPlatform
 			&& (Target.Configuration == UnrealTargetConfiguration.Debug || Target.Configuration == UnrealTargetConfiguration.Development);
 	}
 }
