[Engine]
* Changed so statestream manager has the same life time as render module
* Added placeholder code for state stream lanes. Typically one UWorld will have one lane but there can also be other lanes running at independent rates
* Moved debug rendering to be more localized

#Synced-CL 50091045

diff --git a/Engine/Plugins/Runtime/nDisplay/Source/DisplayCluster/Private/Game/EngineClasses/Basics/DisplayClusterViewportClient.cpp b/Engine/Plugins/Runtime/nDisplay/Source/DisplayCluster/Private/Game/EngineClasses/Basics/DisplayClusterViewportClient.cpp
--- a/Engine/Plugins/Runtime/nDisplay/Source/DisplayCluster/Private/Game/EngineClasses/Basics/DisplayClusterViewportClient.cpp
+++ b/Engine/Plugins/Runtime/nDisplay/Source/DisplayCluster/Private/Game/EngineClasses/Basics/DisplayClusterViewportClient.cpp
@@ -755,6 +755,12 @@ void UDisplayClusterViewportClient::Draw(FViewport* InViewport, FCanvas* SceneCa
 				Y += 20;
 			}
 
+			virtual void DrawText(uint32 InX, uint32 InY, const FStringView& Text, const FColor& Color) override
+			{
+				Canvas->SetDrawColor(Color);
+				Canvas->DrawText(GEngine->GetSmallFont(), Text, InX, InY);
+			}
+
 			uint32 Y = 200;
 			FViewport* Viewport;
 			UCanvas* Canvas;
diff --git a/Engine/Source/Programs/LowLevelTests/StateStreamTests/Public/MockStateStream.h b/Engine/Source/Programs/LowLevelTests/StateStreamTests/Public/MockStateStream.h
--- a/Engine/Source/Programs/LowLevelTests/StateStreamTests/Public/MockStateStream.h
+++ b/Engine/Source/Programs/LowLevelTests/StateStreamTests/Public/MockStateStream.h
@@ -61,5 +61,5 @@ class IMockStateStream
 {
 public:
 	DECLARE_STATESTREAM(Mock)
-	virtual Handle Game_CreateInstance(const FMockStaticState& Ss, const FMockDynamicState& Ds) = 0;
+	virtual Handle Game_CreateInstance(const FMockStaticState& Ss, const FMockDynamicState& Ds, uint32 LaneId = DefaultLaneId) = 0;
 };
diff --git a/Engine/Source/Programs/LowLevelTests/StateStreamTests/Tests/StateStreamTests.cpp b/Engine/Source/Programs/LowLevelTests/StateStreamTests/Tests/StateStreamTests.cpp
--- a/Engine/Source/Programs/LowLevelTests/StateStreamTests/Tests/StateStreamTests.cpp
+++ b/Engine/Source/Programs/LowLevelTests/StateStreamTests/Tests/StateStreamTests.cpp
@@ -84,7 +84,8 @@ TEST_CASE("StateStream::SingleInstance", "[StateStream]")
 	FMockStateStreamImpl MockStateStreamImpl;
 	FStateStreamManagerImpl ManagerImpl;
 	ManagerImpl.Render_Register(MockStateStreamImpl, false);
-
+	uint32 LaneId = ManagerImpl.Game_CreateLane();
+	ManagerImpl.Game_SetDefaultLane(LaneId);
 
 	// Game thread populating ticks
 	{
@@ -92,29 +93,29 @@ TEST_CASE("StateStream::SingleInstance", "[StateStream]")
 		IMockStateStream& Stream = Manager.Game_Get<IMockStateStream>();
 
 		// Tick 1, create an instance
-		Manager.Game_BeginTick();
+		Manager.Game_BeginTick(LaneId);
 		FMockHandle Handle = Stream.Game_CreateInstance(FMockStaticState(), FMockDynamicState{0});
-		Manager.Game_EndTick(100);
+		Manager.Game_EndTick(100, LaneId);
 
 		// Tick 2, update the instance. Set Value to 100
-		Manager.Game_BeginTick();
+		Manager.Game_BeginTick(LaneId);
 		Handle.Update(FMockDynamicState{100, true});
 		Handle.Update(FMockDynamicState{100});
-		Manager.Game_EndTick(200);
+		Manager.Game_EndTick(200, LaneId);
 
 		// Tick 3, update the instance again. Set Value to 200
-		Manager.Game_BeginTick();
+		Manager.Game_BeginTick(LaneId);
 		Handle.Update(FMockDynamicState{200, false});
-		Manager.Game_EndTick(300);
+		Manager.Game_EndTick(300, LaneId);
 
 		// Tick 4, destroy the instance
-		Manager.Game_BeginTick();
+		Manager.Game_BeginTick(LaneId);
 		Handle = {};
-		Manager.Game_EndTick(400);
+		Manager.Game_EndTick(400, LaneId);
 
 		// Tick 5, empty
-		Manager.Game_BeginTick();
-		Manager.Game_EndTick(500);
+		Manager.Game_BeginTick(LaneId);
+		Manager.Game_EndTick(500, LaneId);
 	}
 
 
@@ -279,30 +280,32 @@ TEST_CASE("StateStream::MultipleInstances", "[StateStream]")
 	FMockStateStreamImpl MockStateStreamImpl;
 	FStateStreamManagerImpl ManagerImpl;
 	ManagerImpl.Render_Register(MockStateStreamImpl, false);
+	uint32 LaneId = ManagerImpl.Game_CreateLane();
+	ManagerImpl.Game_SetDefaultLane(LaneId);
 
 	{
 		IStateStreamManager& Manager = ManagerImpl;
 		IMockStateStream& Stream = Manager.Game_Get<IMockStateStream>();
 
 		// Tick 1
-		Manager.Game_BeginTick();
+		Manager.Game_BeginTick(LaneId);
 		FMockHandle Handle1 = Stream.Game_CreateInstance(FMockStaticState(), FMockDynamicState{0});
 		FMockHandle Handle2 = Stream.Game_CreateInstance(FMockStaticState(), FMockDynamicState{0});
-		Manager.Game_EndTick(100);
+		Manager.Game_EndTick(100, LaneId);
 
 		// Tick 2
-		Manager.Game_BeginTick();
+		Manager.Game_BeginTick(LaneId);
 		Handle1 = {};
-		Manager.Game_EndTick(200);
+		Manager.Game_EndTick(200, LaneId);
 
 		// Tick 3
-		Manager.Game_BeginTick();
+		Manager.Game_BeginTick(LaneId);
 		Handle2 = {};
-		Manager.Game_EndTick(300);
+		Manager.Game_EndTick(300, LaneId);
 
 		// Tick 4
-		Manager.Game_BeginTick();
-		Manager.Game_EndTick(400);
+		Manager.Game_BeginTick(LaneId);
+		Manager.Game_EndTick(400, LaneId);
 	}
 
 	FStateStreamManagerImpl& Manager = ManagerImpl;
@@ -338,30 +341,32 @@ TEST_CASE("StateStream::GarbageCollect", "[StateStream]")
 	FMockStateStreamImpl MockStateStreamImpl;
 	FStateStreamManagerImpl ManagerImpl;
 	ManagerImpl.Render_Register(MockStateStreamImpl, false);
+	uint32 LaneId = ManagerImpl.Game_CreateLane();
+	ManagerImpl.Game_SetDefaultLane(LaneId);
 
 	{
 		IStateStreamManager& Manager = ManagerImpl;
 		IMockStateStream& Stream = Manager.Game_Get<IMockStateStream>();
 
 		// Tick 1
-		Manager.Game_BeginTick();
+		Manager.Game_BeginTick(LaneId);
 		FMockHandle Handle1 = Stream.Game_CreateInstance(FMockStaticState(), FMockDynamicState{0});
 		FMockHandle Handle2 = Stream.Game_CreateInstance(FMockStaticState(), FMockDynamicState{0});
-		Manager.Game_EndTick(100);
+		Manager.Game_EndTick(100, LaneId);
 		CHECK_EQUAL(MockStateStreamImpl.GetUsedInstancesCount(), 2);
-		CHECK_EQUAL(MockStateStreamImpl.GetUsedDynamicstatesCount(), 2);
+		CHECK_EQUAL(MockStateStreamImpl.GetUsedDynamicStatesCount(), 2);
 
 		// Tick 2
-		Manager.Game_BeginTick();
+		Manager.Game_BeginTick(LaneId);
 		Handle1 = {};
 		Handle2 = {};
-		Manager.Game_EndTick(200);
+		Manager.Game_EndTick(200, LaneId);
 
 		// Tick 3
-		Manager.Game_BeginTick();
-		Manager.Game_EndTick(300);
+		Manager.Game_BeginTick(LaneId);
+		Manager.Game_EndTick(300, LaneId);
 		CHECK_EQUAL(MockStateStreamImpl.GetUsedInstancesCount(), 2);
-		CHECK_EQUAL(MockStateStreamImpl.GetUsedDynamicstatesCount(), 2);
+		CHECK_EQUAL(MockStateStreamImpl.GetUsedDynamicStatesCount(), 2);
 	}
 
 	FStateStreamManagerImpl& Manager = ManagerImpl;
@@ -369,15 +374,15 @@ TEST_CASE("StateStream::GarbageCollect", "[StateStream]")
 	{
 		Manager.Render_Update(10);
 		Manager.Render_GarbageCollect();
-		CHECK_EQUAL(MockStateStreamImpl.GetUsedDynamicstatesCount(), 2);
+		CHECK_EQUAL(MockStateStreamImpl.GetUsedDynamicStatesCount(), 2);
 
 		Manager.Render_Update(100);
 		Manager.Render_GarbageCollect();
-		CHECK_EQUAL(MockStateStreamImpl.GetUsedDynamicstatesCount(), 2);
+		CHECK_EQUAL(MockStateStreamImpl.GetUsedDynamicStatesCount(), 2);
 
 		Manager.Render_Update(110);
 		Manager.Render_GarbageCollect();
-		CHECK_EQUAL(MockStateStreamImpl.GetUsedDynamicstatesCount(), 2);
+		CHECK_EQUAL(MockStateStreamImpl.GetUsedDynamicStatesCount(), 2);
 
 		Manager.Render_Update(200);
 		Manager.Render_GarbageCollect();
@@ -385,12 +390,12 @@ TEST_CASE("StateStream::GarbageCollect", "[StateStream]")
 		Manager.Render_Update(210);
 		Manager.Render_GarbageCollect();
 		CHECK_EQUAL(Stream.GetUsedInstancesCount(), 2);
-		CHECK_EQUAL(MockStateStreamImpl.GetUsedDynamicstatesCount(), 2);
+		CHECK_EQUAL(MockStateStreamImpl.GetUsedDynamicStatesCount(), 2);
 
 		Manager.Render_Update(300);
 		Manager.Render_GarbageCollect();
 		CHECK_EQUAL(Stream.GetUsedInstancesCount(), 0);
-		CHECK_EQUAL(MockStateStreamImpl.GetUsedDynamicstatesCount(), 0);
+		CHECK_EQUAL(MockStateStreamImpl.GetUsedDynamicStatesCount(), 0);
 	}
 }
 
@@ -398,10 +403,10 @@ TEST_CASE("StateStream::InternalDependencies", "[StateStream]")
 {
 	FTransformStateStreamImpl TransformStreamImpl;
 
-
 	FStateStreamManagerImpl ManagerImpl;
 	ManagerImpl.Render_Register(TransformStreamImpl, false);
-
+	uint32 LaneId = ManagerImpl.Game_CreateLane();
+	ManagerImpl.Game_SetDefaultLane(LaneId);
 
 	{
 		IStateStreamManager& Manager = ManagerImpl;
@@ -410,29 +415,29 @@ TEST_CASE("StateStream::InternalDependencies", "[StateStream]")
 		FTransformHandle TransformHandle;
 
 		// Tick 1
-		Manager.Game_BeginTick();
+		Manager.Game_BeginTick(LaneId);
 		{
 			FTransformDynamicState Ds;
 			Ds.SetParent(TransformStream.Game_CreateInstance(FTransformStaticState(), FTransformDynamicState{}));
 			TransformHandle = TransformStream.Game_CreateInstance(FTransformStaticState(), Ds);
 		}
-		Manager.Game_EndTick(100);
+		Manager.Game_EndTick(100, LaneId);
 		CHECK_EQUAL(TransformStreamImpl.GetUsedInstancesCount(), 2);
-		CHECK_EQUAL(TransformStreamImpl.GetUsedDynamicstatesCount(), 2);
+		CHECK_EQUAL(TransformStreamImpl.GetUsedDynamicStatesCount(), 2);
 
 		// Tick 2
-		Manager.Game_BeginTick();
+		Manager.Game_BeginTick(LaneId);
 		{
 			FTransformDynamicState Ds;
 			Ds.SetParent(TransformStream.Game_CreateInstance(FTransformStaticState(), FTransformDynamicState{}));
 			TransformHandle.Update(Ds);
 		}
-		Manager.Game_EndTick(200);
+		Manager.Game_EndTick(200, LaneId);
 
 		// Tick 3
-		Manager.Game_BeginTick();
+		Manager.Game_BeginTick(LaneId);
 		TransformHandle = {};
-		Manager.Game_EndTick(300);
+		Manager.Game_EndTick(300, LaneId);
 
 		Manager.Game_Exit();
 	}
@@ -469,6 +474,9 @@ TEST_CASE("StateStream::ExternalDependencies", "[StateStream]")
 	ManagerImpl.Render_Register(TransformStreamImpl, false);
 	ManagerImpl.Render_Register(MockStreamImpl, false);
 
+	uint32 LaneId = ManagerImpl.Game_CreateLane();
+	ManagerImpl.Game_SetDefaultLane(LaneId);
+
 
 	{
 		IStateStreamManager& Manager = ManagerImpl;
@@ -478,31 +486,31 @@ TEST_CASE("StateStream::ExternalDependencies", "[StateStream]")
 		FMockHandle MockHandle;
 
 		// Tick 1
-		Manager.Game_BeginTick();
+		Manager.Game_BeginTick(LaneId);
 		{
 			FMockDynamicState Ds;
 			Ds.SetTransform(TransformStream.Game_CreateInstance(FTransformStaticState(), FTransformDynamicState{}));
 			MockHandle = MockStream.Game_CreateInstance(FMockStaticState(), Ds);
 		}
-		Manager.Game_EndTick(100);
+		Manager.Game_EndTick(100, LaneId);
 		CHECK_EQUAL(TransformStreamImpl.GetUsedInstancesCount(), 1);
-		CHECK_EQUAL(TransformStreamImpl.GetUsedDynamicstatesCount(), 1);
+		CHECK_EQUAL(TransformStreamImpl.GetUsedDynamicStatesCount(), 1);
 		CHECK_EQUAL(MockStreamImpl.GetUsedInstancesCount(), 1);
-		CHECK_EQUAL(MockStreamImpl.GetUsedDynamicstatesCount(), 1);
+		CHECK_EQUAL(MockStreamImpl.GetUsedDynamicStatesCount(), 1);
 
 		// Tick 2
-		Manager.Game_BeginTick();
+		Manager.Game_BeginTick(LaneId);
 		{
 			FMockDynamicState Ds;
 			Ds.SetTransform(TransformStream.Game_CreateInstance(FTransformStaticState(), FTransformDynamicState{}));
 			MockHandle.Update(Ds);
 		}
-		Manager.Game_EndTick(200);
+		Manager.Game_EndTick(200, LaneId);
 
 		// Tick 3
-		Manager.Game_BeginTick();
+		Manager.Game_BeginTick(LaneId);
 		MockHandle = {};
-		Manager.Game_EndTick(300);
+		Manager.Game_EndTick(300, LaneId);
 
 		Manager.Game_Exit();
 	}
@@ -539,26 +547,29 @@ TEST_CASE("StateStream::Interpolation", "[StateStream]")
 	FStateStreamManagerImpl ManagerImpl;
 	ManagerImpl.Render_Register(TransformStreamImpl, false);
 
+	uint32 LaneId = ManagerImpl.Game_CreateLane();
+	ManagerImpl.Game_SetDefaultLane(LaneId);
+
 	{
 		IStateStreamManager& Manager = ManagerImpl;
 		ITransformStateStream& Stream = Manager.Game_Get<ITransformStateStream>();
 
 		// Tick 1
-		Manager.Game_BeginTick();
+		Manager.Game_BeginTick(LaneId);
 		FTransformHandle Handle = Stream.Game_CreateInstance(FTransformStaticState(), FTransformDynamicState());
-		Manager.Game_EndTick(0);
+		Manager.Game_EndTick(0, LaneId);
 
 		// Tick 2
-		Manager.Game_BeginTick();
+		Manager.Game_BeginTick(LaneId);
 		FTransformDynamicState newState;
 		newState.SetLocalTransform(FTransform(FQuat::Identity, FVector(100, 100, 100)));
 		Handle.Update(newState);
-		Manager.Game_EndTick(100);
+		Manager.Game_EndTick(100, LaneId);
 
 		// Tick 3
-		Manager.Game_BeginTick();
+		Manager.Game_BeginTick(LaneId);
 		Handle = {};
-		Manager.Game_EndTick(200);
+		Manager.Game_EndTick(200, LaneId);
 	}
 
 	FStateStreamManagerImpl& Manager = ManagerImpl;
diff --git a/Engine/Source/Runtime/Engine/Classes/Engine/World.h b/Engine/Source/Runtime/Engine/Classes/Engine/World.h
--- a/Engine/Source/Runtime/Engine/Classes/Engine/World.h
+++ b/Engine/Source/Runtime/Engine/Classes/Engine/World.h
@@ -1580,6 +1580,7 @@ private:
 #if WITH_STATE_STREAM
 	/** Manager for all existing state streams */
 	IStateStreamManager* StateStreamManager = nullptr;
+	uint32 LaneId;
 #endif
 
 public:
@@ -1604,6 +1605,7 @@ public:
 
 	/** Get statestream manager */
 	IStateStreamManager* GetStateStreamManager() { return StateStreamManager; }
+	uint32 GetLaneId() { return LaneId; }
 
 #endif
 
diff --git a/Engine/Source/Runtime/Engine/Private/GameViewportClient.cpp b/Engine/Source/Runtime/Engine/Private/GameViewportClient.cpp
--- a/Engine/Source/Runtime/Engine/Private/GameViewportClient.cpp
+++ b/Engine/Source/Runtime/Engine/Private/GameViewportClient.cpp
@@ -1991,11 +1991,17 @@ void UGameViewportClient::Draw(FViewport* InViewport, FCanvas* SceneCanvas)
 				Y += 20;
 			}
 
+			virtual void DrawText(uint32 InX, uint32 InY, const FStringView& Text, const FColor& Color) override
+			{
+				Canvas->SetDrawColor(Color);
+				Canvas->DrawText(GEngine->GetSmallFont(), Text, InX, InY);
+			}
+
 			uint32 Y = 200;
 			FViewport* Viewport;
 			UCanvas* Canvas;
 		} StateStreamDebugRenderer(InViewport, DebugCanvasObject);
-		Manager->Game_DebugRender(StateStreamDebugRenderer);
+		Manager->Game_DebugRender(StateStreamDebugRenderer, 0);
 	}
 #endif
 
diff --git a/Engine/Source/Runtime/Engine/Private/LevelTick.cpp b/Engine/Source/Runtime/Engine/Private/LevelTick.cpp
--- a/Engine/Source/Runtime/Engine/Private/LevelTick.cpp
+++ b/Engine/Source/Runtime/Engine/Private/LevelTick.cpp
@@ -1123,6 +1123,10 @@ void UWorld::SendAllEndOfFrameUpdatesInternal(EEndOfFrameUpdateLocation Location
 	CSV_SCOPED_TIMING_STAT_EXCLUSIVE(EndOfFrameUpdates);
 	CSV_SCOPED_SET_WAIT_STAT(EndOfFrameUpdates);
 
+#if WITH_STATE_STREAM
+	FStateStreamLaneScope LaneScope(StateStreamManager, LaneId);
+#endif
+
 	const bool bAsyncInTick = Location == EEndOfFrameUpdateLocation::Async && bInTick;
 
 	if (!bAsyncInTick)
@@ -1310,6 +1314,10 @@ void UWorld::StartAsyncSendAllEndOfFrameUpdates()
 
 void UWorld::FinishAsyncSendAllEndOfFrameUpdates()
 {
+#if WITH_STATE_STREAM
+	FStateStreamLaneScope LaneScope(StateStreamManager, LaneId);
+#endif
+
 	if (!EndOfFrameUpdateContexts.IsEmpty())
 	{
 		Scene->GetSkeletalMeshUpdater()->EndAsyncPushCommands();
@@ -1486,6 +1494,10 @@ void UWorld::Tick( ELevelTick TickType, float DeltaSeconds )
 		return;
 	}
 
+#if WITH_STATE_STREAM
+	FStateStreamLaneScope LaneScope(StateStreamManager, LaneId);
+#endif
+
 	UE::Stats::FThreadIdleStats::BeginCriticalPath();
 
 	RHI_BREADCRUMB_EVENT_GAMETHREAD("WorldTick");
@@ -2031,8 +2043,8 @@ void UWorld::Tick( ELevelTick TickType, float DeltaSeconds )
 	FWorldDelegates::OnWorldTickEnd.Broadcast(this, TickType, DeltaSeconds);
 
 #if WITH_STATE_STREAM
-	StateStreamManager->Game_EndTick(RealTimeSeconds);
-	StateStreamManager->Game_BeginTick();
+	StateStreamManager->Game_EndTick(RealTimeSeconds, LaneId);
+	StateStreamManager->Game_BeginTick(LaneId);
 #endif
 
 	UE::Stats::FThreadIdleStats::EndCriticalPath();
diff --git a/Engine/Source/Runtime/Engine/Private/UnrealEngine.cpp b/Engine/Source/Runtime/Engine/Private/UnrealEngine.cpp
--- a/Engine/Source/Runtime/Engine/Private/UnrealEngine.cpp
+++ b/Engine/Source/Runtime/Engine/Private/UnrealEngine.cpp
@@ -28,6 +28,7 @@ UnrealEngine.cpp: Implements the UEngine class and helpers.
 #include "HAL/PlatformFile.h"
 #include "Misc/OutputDeviceFile.h"
 #include "Math/ColorList.h"
+#include "StateStreamManager.h"
 #include "Stats/StatsMisc.h"
 #include "Stats/ThreadIdleStats.h"
 #include "Misc/AutomationTest.h"
@@ -16181,6 +16182,10 @@ bool UEngine::LoadMap( FWorldContext& WorldContext, FURL URL, class UPendingNetG
 		}
 	}
 
+#if WITH_STATE_STREAM
+	FStateStreamLaneScope LaneScope(NewWorld->GetStateStreamManager(), NewWorld->GetLaneId());
+#endif
+
 	// Handle pending level.
 	if( Pending )
 	{
diff --git a/Engine/Source/Runtime/Engine/Private/World.cpp b/Engine/Source/Runtime/Engine/Private/World.cpp
--- a/Engine/Source/Runtime/Engine/Private/World.cpp
+++ b/Engine/Source/Runtime/Engine/Private/World.cpp
@@ -1585,6 +1585,10 @@ void UWorld::FinishDestroy()
 
 		ReleasePhysicsScene();
 
+#if WITH_STATE_STREAM
+		StateStreamManager->Game_DestroyLane(LaneId);
+#endif
+
 		if (Scene)
 		{
 			Scene->Release();
@@ -1600,7 +1604,6 @@ void UWorld::FinishDestroy()
 #if WITH_STATE_STREAM
 	if (StateStreamManager)
 	{
-		GetRendererModule().FreeStateStream(StateStreamManager);
 		StateStreamManager = nullptr;
 	}
 #endif
@@ -2442,8 +2445,10 @@ void UWorld::InitWorld(const InitializationValues IVS)
 	}
 
 #if WITH_STATE_STREAM
-	StateStreamManager = GetRendererModule().AllocateStateStream(this);
-	StateStreamManager->Game_BeginTick();
+	StateStreamManager = GetRendererModule().GetStateStreamManager();
+	LaneId = StateStreamManager->Game_CreateLane();
+	StateStreamManager->Game_SetLaneUserData(LaneId, Scene);
+	StateStreamManager->Game_BeginTick(LaneId);
 #endif
 
 	// Prepare AI systems
@@ -5914,6 +5919,11 @@ void UWorld::InitializeActorsForPlay(const FURL& InURL, bool bResetTime, FRegist
 	SCOPED_BOOT_TIMING("UWorld::InitializeActorsForPlay");
 	double StartTime = FPlatformTime::Seconds();
 
+
+#if WITH_STATE_STREAM
+	FStateStreamLaneScope LaneScope(StateStreamManager, LaneId);
+#endif
+
 	// Don't reset time for seamless world transitions.
 	if (bResetTime)
 	{
@@ -6320,7 +6330,7 @@ void UWorld::CleanupWorldInternal(bool bSessionEnded, bool bCleanupResources, bo
 #if WITH_STATE_STREAM
 	if (StateStreamManager)
 	{
-		StateStreamManager->Game_EndTick(RealTimeSeconds);
+		StateStreamManager->Game_EndTick(RealTimeSeconds, LaneId);
 	}
 #endif
 
@@ -10102,6 +10112,10 @@ void UWorld::RecreateScene(ERHIFeatureLevel::Type InFeatureLevel, bool bBroadcas
 
 		IRendererModule& RendererModule = GetRendererModule();
 		RendererModule.AllocateScene(this, bRequiresHitProxies, bCreateFXSystem, InFeatureLevel);
+		
+#if WITH_STATE_STREAM
+		StateStreamManager->Game_SetLaneUserData(LaneId, Scene);
+#endif
 
 		for (ULevel* Level : Levels)
 		{
diff --git a/Engine/Source/Runtime/Engine/Public/StateStream/ExponentialHeightFogStateStream.h b/Engine/Source/Runtime/Engine/Public/StateStream/ExponentialHeightFogStateStream.h
--- a/Engine/Source/Runtime/Engine/Public/StateStream/ExponentialHeightFogStateStream.h
+++ b/Engine/Source/Runtime/Engine/Public/StateStream/ExponentialHeightFogStateStream.h
@@ -109,7 +109,7 @@ class IExponentialHeightFogStateStream
 {
 public:
 	DECLARE_STATESTREAM(ExponentialHeightFog)
-	virtual FExponentialHeightFogHandle Game_CreateInstance(const FExponentialHeightFogStaticState& Ss, const FExponentialHeightFogDynamicState& Ds) = 0;
+	virtual FExponentialHeightFogHandle Game_CreateInstance(const FExponentialHeightFogStaticState& Ss, const FExponentialHeightFogDynamicState& Ds, uint32 LaneId = DefaultLaneId) = 0;
 };
 
 ////////////////////////////////////////////////////////////////////////////////////////////////////
diff --git a/Engine/Source/Runtime/Engine/Public/StateStream/ParticleSystemStateStream.h b/Engine/Source/Runtime/Engine/Public/StateStream/ParticleSystemStateStream.h
--- a/Engine/Source/Runtime/Engine/Public/StateStream/ParticleSystemStateStream.h
+++ b/Engine/Source/Runtime/Engine/Public/StateStream/ParticleSystemStateStream.h
@@ -50,7 +50,7 @@ class IParticleSystemStateStream
 {
 public:
 	DECLARE_STATESTREAM(ParticleSystem)
-	virtual FParticleSystemHandle Game_CreateInstance(const FParticleSystemStaticState& Ss, const FParticleSystemDynamicState& Ds) = 0;
+	virtual FParticleSystemHandle Game_CreateInstance(const FParticleSystemStaticState& Ss, const FParticleSystemDynamicState& Ds, uint32 LaneId = DefaultLaneId) = 0;
 
 	virtual void SetOtherBackend(IParticleSystemStateStream* Other) {}
 };
diff --git a/Engine/Source/Runtime/Engine/Public/StateStream/SkinnedMeshStateStream.h b/Engine/Source/Runtime/Engine/Public/StateStream/SkinnedMeshStateStream.h
--- a/Engine/Source/Runtime/Engine/Public/StateStream/SkinnedMeshStateStream.h
+++ b/Engine/Source/Runtime/Engine/Public/StateStream/SkinnedMeshStateStream.h
@@ -58,7 +58,7 @@ class ISkinnedMeshStateStream
 {
 public:
 	DECLARE_STATESTREAM(SkinnedMesh)
-	virtual FSkinnedMeshHandle Game_CreateInstance(const FSkinnedMeshStaticState& Ss, const FSkinnedMeshDynamicState& Ds) = 0;
+	virtual FSkinnedMeshHandle Game_CreateInstance(const FSkinnedMeshStaticState& Ss, const FSkinnedMeshDynamicState& Ds, uint32 LaneId = DefaultLaneId) = 0;
 };
 
 ////////////////////////////////////////////////////////////////////////////////////////////////////
diff --git a/Engine/Source/Runtime/Engine/Public/StateStream/SkyAtmosphereStateStream.h b/Engine/Source/Runtime/Engine/Public/StateStream/SkyAtmosphereStateStream.h
--- a/Engine/Source/Runtime/Engine/Public/StateStream/SkyAtmosphereStateStream.h
+++ b/Engine/Source/Runtime/Engine/Public/StateStream/SkyAtmosphereStateStream.h
@@ -132,5 +132,5 @@ class ISkyAtmosphereStateStream
 {
 public:
 	DECLARE_STATESTREAM(SkyAtmosphere)
-	virtual FSkyAtmosphereHandle Game_CreateInstance(const FSkyAtmosphereStaticState& Ss, const FSkyAtmosphereDynamicState& Ds) = 0;
+	virtual FSkyAtmosphereHandle Game_CreateInstance(const FSkyAtmosphereStaticState& Ss, const FSkyAtmosphereDynamicState& Ds, uint32 LaneId = DefaultLaneId) = 0;
 };
diff --git a/Engine/Source/Runtime/Engine/Public/StateStream/StaticMeshStateStream.h b/Engine/Source/Runtime/Engine/Public/StateStream/StaticMeshStateStream.h
--- a/Engine/Source/Runtime/Engine/Public/StateStream/StaticMeshStateStream.h
+++ b/Engine/Source/Runtime/Engine/Public/StateStream/StaticMeshStateStream.h
@@ -63,5 +63,5 @@ class IStaticMeshStateStream
 {
 public:
 	DECLARE_STATESTREAM(StaticMesh)
-	virtual FStaticMeshHandle Game_CreateInstance(const FStaticMeshStaticState& Ss, const FStaticMeshDynamicState& Ds) = 0;
+	virtual FStaticMeshHandle Game_CreateInstance(const FStaticMeshStaticState& Ss, const FStaticMeshDynamicState& Ds, uint32 LaneId = DefaultLaneId) = 0;
 };
diff --git a/Engine/Source/Runtime/RenderCore/Private/RenderCommandStateStream.cpp b/Engine/Source/Runtime/RenderCore/Private/RenderCommandStateStream.cpp
--- a/Engine/Source/Runtime/RenderCore/Private/RenderCommandStateStream.cpp
+++ b/Engine/Source/Runtime/RenderCore/Private/RenderCommandStateStream.cpp
@@ -11,25 +11,39 @@
 #include "StateStreamDebugRenderer.h"
 #include "StateStreamManagerImpl.h"
 
+#define UE_RENDER_COMMAND_DEBUG_DRAW 1
+
 ////////////////////////////////////////////////////////////////////////////////////////////////////
 
+class FRenderCommandStateStream;
+FRenderCommandStateStream* g_renderCommandStateStream;
+
+
 class FRenderCommandStateStream : public IStateStream
 {
 public:
-	FRenderCommandStateStream(FSceneInterface& InScene)
+	FRenderCommandStateStream()
 	{
+		g_renderCommandStateStream = this;
 		// TODO: Flush
 		//FRenderCommandFence Fence;
 		//Fence.BeginFence();
 		//Fence.Wait(false);
 	}
-
-	virtual void Game_BeginTick() override
+	~FRenderCommandStateStream()
 	{
+		g_renderCommandStateStream = nullptr;
 	}
-	virtual void Game_EndTick(StateStreamTime AbsoluteTime) override
+	virtual void Game_BeginTick(uint32 LaneId) override
 	{
 	}
+	virtual void Game_EndTick(StateStreamTime AbsoluteTime, uint32 LaneId) override
+	{
+	}
+	virtual void Game_SetDefaultLane(uint32 LaneId) override
+	{
+		ActiveDefaultLaneId = LaneId;
+	}
 	virtual void Game_Exit() override
 	{
 	}
@@ -42,6 +56,32 @@ public:
 		TStringBuilder<1024> DebugLine;
 		DebugLine.Appendf(TEXT("State stream : [%s]"), GetDebugName());
 		Renderer.DrawText(*DebugLine, FColor::Orange);
+
+#if UE_RENDER_COMMAND_DEBUG_DRAW
+		uint32 X = 500;
+		uint32 Y = 100;
+
+		DebugLine.Reset();
+
+		FScopeLock Lock(&RenderCommandsLock);
+		DebugLine.Appendf(TEXT("Num RenderCommands = %d"), RenderCommands.Num());
+		Renderer.DrawText(X, Y, *DebugLine, FColor::Blue);
+		Y += 20;
+
+		for (auto& Pair : RenderCommands)
+		{
+			if (Y == 1000)
+			{
+				Y = 100;
+				X += 300;
+			}
+			DebugLine.Reset();
+			DebugLine.Appendf(TEXT("%s, %d"), *Pair.Key, Pair.Value);
+			Renderer.DrawText(X, Y, *DebugLine, FColor::Orange);
+			Y += 20;
+		}
+		RenderCommands.Empty();
+#endif
 	}
 
 	virtual void Render_Update(StateStreamTime AbsoluteTime) override
@@ -65,35 +105,49 @@ public:
 	}
 	virtual const TCHAR* GetDebugName() override { return TEXT("RenderCommand"); }
 
+	bool AddCommand(TUniqueFunction<void(FRHICommandListImmediate&)>& Function, const FRenderCommandTag& Tag)
+	{
+		// If no lane is created, then we just flush things through... when lanes are created we wait and clear RC
+
+
+		if (IsInGameThread() || IsInParallelGameThread())
+		{
+			//check(g_renderCommandStateStream->ActiveDefaultLaneId != DefaultLaneId);
+
+#if UE_RENDER_COMMAND_DEBUG_DRAW
+			FScopeLock Lock(&RenderCommandsLock);
+			++RenderCommands.FindOrAdd(Tag.GetName(), 0);
+#endif
+
+			// These should be added to current tick of RenderCommandStateStream
+		}
+		else // if (GRenderingThreadHeartbeat != FRunnableThread::GetRunnableThread())
+		{
+			// These can be called in between game ticks but still needs to be in sync with ticks.. will need to queue up separately
+			// and inject to beginning of next tick when begin?
+
+			// Things that could end up here
+			//   GRenderingThreadHeartbeat (FDefaultGameMoviePlayer, etc)
+			//   AsyncPackage2 loading (FMaterialShaderMap, FTextureReference, etc)
+			//   FSlateRHIRenderer::ReleaseDrawBuffer
+		}
+		return false;
+	}
+
 	enum { Id = RenderCommandStateStreamId };
 
-	//TMap<FString,uint32> DebugElementsCounts;
+	uint32 ActiveDefaultLaneId = DefaultLaneId;
+
+	FCriticalSection RenderCommandsLock;
+	TMap<FString,uint32> RenderCommands;
 };
 
 ////////////////////////////////////////////////////////////////////////////////////////////////////
 
 bool RenderCommandStateStream_AddCommand(TUniqueFunction<void(FRHICommandListImmediate&)>& Function, const FRenderCommandTag& Tag)
 {
-	if (IsInGameThread() || IsInParallelGameThread())
-	{
-		//if (RenderCommandStateStream_AddCommand(Function, Tag))
-		//	return;
-		FStateStreamManagerImpl::GetDebugRenderCommands().AddDebugElements(Tag.GetName(), 1);
-
-		// These should be added to current tick of RenderCommandStateStream
-
-	}
-	else // if (GRenderingThreadHeartbeat != FRunnableThread::GetRunnableThread())
-	{
-		// These can be called in between game ticks but still needs to be in sync with ticks.. will need to queue up separately
-		// and inject to beginning of next tick when begin?
-
-		// Things that could end up here
-		//   GRenderingThreadHeartbeat (FDefaultGameMoviePlayer, etc)
-		//   AsyncPackage2 loading (FMaterialShaderMap, FTextureReference, etc)
-		//   FSlateRHIRenderer::ReleaseDrawBuffer
-	}
-	return false;
+	check(g_renderCommandStateStream);
+	return g_renderCommandStateStream->AddCommand(Function, Tag);
 }
 
 ////////////////////////////////////////////////////////////////////////////////////////////////////
diff --git a/Engine/Source/Runtime/RenderCore/Public/RendererInterface.h b/Engine/Source/Runtime/RenderCore/Public/RendererInterface.h
--- a/Engine/Source/Runtime/RenderCore/Public/RendererInterface.h
+++ b/Engine/Source/Runtime/RenderCore/Public/RendererInterface.h
@@ -709,8 +709,8 @@ public:
 	virtual void RemoveScene(FSceneInterface* Scene) = 0;
 
 #if WITH_STATE_STREAM
-	virtual IStateStreamManager* AllocateStateStream(UWorld* World) = 0;
-	virtual void FreeStateStream(IStateStreamManager* Manager) = 0;
+	virtual IStateStreamManager* GetStateStreamManager() = 0;
+	virtual void DebugLogRenderProxies(FName Type, uint32 Num) = 0;
 #endif
 
 	/**
diff --git a/Engine/Source/Runtime/Renderer/Private/Renderer.cpp b/Engine/Source/Runtime/Renderer/Private/Renderer.cpp
--- a/Engine/Source/Runtime/Renderer/Private/Renderer.cpp
+++ b/Engine/Source/Runtime/Renderer/Private/Renderer.cpp
@@ -6,6 +6,9 @@
 
 #include "CoreMinimal.h"
 #include "Misc/CoreMisc.h"
+#include "StateStreamCreator.h"
+#include "StateStreamDebugRenderer.h"
+#include "StateStreamManagerImpl.h"
 #include "Stats/Stats.h"
 #include "Modules/ModuleManager.h"
 #include "Async/TaskGraphInterfaces.h"
@@ -78,6 +81,54 @@ static FAutoConsoleVariableRef CVarBindTileMeshDrawingDummyRenderTarget(
 	TEXT("Enable binding of a dummy render target for tile mesh drawing to workaround driver bugs.")
 );
 
+#if WITH_STATE_STREAM
+class FRendererModule::FRendererStateStreamManager : public FStateStreamManagerImpl
+{
+public:
+	virtual void Game_DebugRender(IStateStreamDebugRenderer& Renderer, uint32 LaneId) override
+	{
+		FStateStreamManagerImpl::Game_DebugRender(Renderer, LaneId);
+
+		TStringBuilder<1024> DebugLine;
+
+		FScopeLock Lock(&ProxyTypeAndCountLock);
+		const int32 NumElementCounts = ProxyTypeAndCount.Num();
+		if (NumElementCounts == 0)
+		{
+			return;
+		}
+		const float NumElementColumns = FMath::Min(10.0f, NumElementCounts);
+		const int32 NumElementRows = FMath::CeilToInt32(NumElementCounts / NumElementColumns);
+
+		// Display the number of elements
+		DebugLine.Reset();
+		DebugLine.Appendf(TEXT("Num Render proxies = %d"), NumElementCounts);
+		Renderer.DrawText(*DebugLine, FColor::Blue);
+
+		// Display per element informations
+		auto ElementIterator = ProxyTypeAndCount.CreateConstIterator();
+		for (int32 ColumnIndex = 0; ColumnIndex < NumElementColumns; ++ColumnIndex)
+		{
+			DebugLine.Reset();
+			//DebugLine.Appendf(TEXT("%s : "), *DebugName);
+			for (int32 RowIndex = 0; RowIndex < NumElementRows; ++RowIndex)
+			{
+				if (ElementIterator)
+				{
+					DebugLine.Appendf(TEXT("[%s | %d]"), *ElementIterator.Key().ToString(), ElementIterator.Value());
+					//UE_LOG(LogTemp, Log, TEXT("%s : [%s | %d]"), *DebugName, *ElementIterator.Key(), ElementIterator.Value());
+					++ElementIterator;
+				}
+			}
+			Renderer.DrawText(*DebugLine, FColor::Orange);
+		}
+	}
+
+	FCriticalSection ProxyTypeAndCountLock;
+	TMap<FName, uint32> ProxyTypeAndCount;
+};
+#endif
+
 void FRendererModule::StartupModule()
 {
 #if MESH_DRAW_COMMAND_STATS
@@ -133,12 +184,23 @@ void FRendererModule::StartupModule()
 			InitVisualizeColorGradingExtension();
 		});
 	}
+
+#if WITH_STATE_STREAM
+	StateStreamManager = new FRendererStateStreamManager();
+	FStateStreamCreator::RegisterStateStreams(FStateStreamRegisterContext{*StateStreamManager});
+	//uint32 RenderCommandsLane = StateStreamManager->Game_CreateLane(); // Create the lane that non-tick based render commands can be stored in
+	//check(RenderCommandsLane == 0);(void)RenderCommandsLane;
+#endif
 }
 
 void FRendererModule::ShutdownModule()
 {
 	UnregisterStopRenderingThreadDelegate(StopRenderingThreadDelegate);
 
+#if WITH_STATE_STREAM
+	delete StateStreamManager;
+#endif
+
 #if RHI_RAYTRACING
 	Nanite::GRayTracingManager.Shutdown();
 
@@ -178,6 +240,19 @@ void FRendererModule::OnWorldCleanup(UWorld* World, bool bSessionEnded, bool bCl
 
 }
 
+#if WITH_STATE_STREAM
+IStateStreamManager* FRendererModule::GetStateStreamManager()
+{
+	return StateStreamManager;
+}
+
+void FRendererModule::DebugLogRenderProxies(FName Type, uint32 Num)
+{
+	FScopeLock Lock(&StateStreamManager->ProxyTypeAndCountLock);
+	StateStreamManager->ProxyTypeAndCount.FindOrAdd(Type, 0) += Num;
+}
+#endif
+
 void FRendererModule::InitializeSystemTextures(FRHICommandListImmediate& RHICmdList)
 {
 	GSystemTextures.InitializeTextures(RHICmdList, GMaxRHIFeatureLevel);
diff --git a/Engine/Source/Runtime/Renderer/Private/RendererModule.h b/Engine/Source/Runtime/Renderer/Private/RendererModule.h
--- a/Engine/Source/Runtime/Renderer/Private/RendererModule.h
+++ b/Engine/Source/Runtime/Renderer/Private/RendererModule.h
@@ -20,6 +20,7 @@ class FSceneViewFamily;
 class FSceneViewStateInterface;
 class FViewInfo;
 class FSceneUniformBuffer;
+class FStateStreamManagerImpl;
 struct FMeshBatch;
 struct FSynthBenchmarkResults;
 struct FSceneTextures;
@@ -157,8 +158,8 @@ public:
 	virtual void AddMeshBatchToGPUScene(FGPUScenePrimitiveCollector* Collector, FMeshBatch& MeshBatch) override;
 
 #if WITH_STATE_STREAM
-	virtual IStateStreamManager* AllocateStateStream(UWorld* World) override;
-	virtual void FreeStateStream(IStateStreamManager* Manager) override;
+	virtual IStateStreamManager* GetStateStreamManager() override;
+	virtual void DebugLogRenderProxies(FName Type, uint32 Num) override;
 #endif
 
 	virtual TUniquePtr<ISceneRenderBuilder> CreateSceneRenderBuilder(FSceneInterface* Scene) override;
@@ -169,6 +170,11 @@ private:
 	FOnPostOpaqueRender OverlayRenderDelegate;
 	FOnResolvedSceneColor PostResolvedSceneColorCallbacks;
 	FDelegateHandle StopRenderingThreadDelegate;
+
+#if WITH_STATE_STREAM
+	class FRendererStateStreamManager;
+	FRendererStateStreamManager* StateStreamManager = nullptr;
+#endif
 };
 
 extern ICustomCulling* GCustomCullingImpl;
diff --git a/Engine/Source/Runtime/Renderer/Private/RendererScene.cpp b/Engine/Source/Runtime/Renderer/Private/RendererScene.cpp
--- a/Engine/Source/Runtime/Renderer/Private/RendererScene.cpp
+++ b/Engine/Source/Runtime/Renderer/Private/RendererScene.cpp
@@ -1372,13 +1372,9 @@ void FScene::BatchAddPrimitivesInternal(TArrayView<T*> InPrimitives)
 #if WITH_STATE_STREAM
 	if (IsInGameThread() || IsInParallelGameThread())
 	{
-		if (FStateStreamManagerImpl* StateStreamManager = static_cast<FStateStreamManagerImpl*>(World->GetStateStreamManager()))
+		if (UObject* PrimitiveObject = ToUObject(InPrimitives.Last()))
 		{
-			if (UObject* PrimitiveObject = ToUObject(InPrimitives.Last()))
-			{
-				const FString ObjectType = PrimitiveObject->GetClass()->GetFName().ToString();
-				StateStreamManager->GetDebugSceneProxies().AddDebugElements(ObjectType, InPrimitives.Num());
-			}
+			GetRendererModule().DebugLogRenderProxies(PrimitiveObject->GetClass()->GetFName(), InPrimitives.Num());
 		}
 	}
 #endif
@@ -4371,8 +4367,8 @@ void FScene::Release()
 		ENQUEUE_RENDER_COMMAND(FReleaseStateStream)(
 			[Scene, StateStreamManager](FRHICommandListImmediate& RHICmdList)
 			{
-				StateStreamManager->Render_Exit();
-				GetRendererModule().FreeStateStream(StateStreamManager);
+				//StateStreamManager->Render_Exit();
+				//GetRendererModule().FreeStateStream(StateStreamManager);
 				if (FSkeletalMeshUpdater* Smu = Scene->SkeletalMeshUpdater)
 				{
 					Smu->Shutdown();
@@ -7003,23 +6999,6 @@ void FRendererModule::RemoveScene(FSceneInterface* Scene)
 	AllocatedScenes.Remove(Scene);
 }
 
-#if WITH_STATE_STREAM
-IStateStreamManager* FRendererModule::AllocateStateStream(UWorld* World)
-{
-	FSceneInterface* Scene = World->Scene;
-	FStateStreamManagerImpl* Manager = new FStateStreamManagerImpl();
-	FStateStreamCreator::RegisterStateStreams(FStateStreamRegisterContext{*Manager, Scene});
-	return Manager;
-}
-
-void FRendererModule::FreeStateStream(IStateStreamManager* Manager)
-{
-	FStateStreamManagerImpl* ManagerImpl = static_cast<FStateStreamManagerImpl*>(Manager);
-	FStateStreamCreator::UnregisterStateStreams(FStateStreamUnregisterContext{*ManagerImpl});
-	delete ManagerImpl;
-}
-#endif
-
 void FRendererModule::UpdateStaticDrawLists()
 {
 	// Update all static meshes in order to recache cached mesh draw commands.
diff --git a/Engine/Source/Runtime/Renderer/Private/StateStream/ExponentialHeightFogStateStreamImpl.cpp b/Engine/Source/Runtime/Renderer/Private/StateStream/ExponentialHeightFogStateStreamImpl.cpp
--- a/Engine/Source/Runtime/Renderer/Private/StateStream/ExponentialHeightFogStateStreamImpl.cpp
+++ b/Engine/Source/Runtime/Renderer/Private/StateStream/ExponentialHeightFogStateStreamImpl.cpp
@@ -6,14 +6,9 @@
 
 ////////////////////////////////////////////////////////////////////////////////////////////////////
 
-FExponentialHeightFogStateStreamImpl::FExponentialHeightFogStateStreamImpl(FSceneInterface& InScene)
-:	Scene(InScene)
-{
-}
-
 void FExponentialHeightFogStateStreamImpl::Render_OnCreate(const FExponentialHeightFogStaticState& Ss, const FExponentialHeightFogDynamicState& Ds, FExponentialHeightFogSceneProxy*& UserData, bool IsDestroyedInSameFrame)
 {
-	Scene.AddExponentialHeightFog(uint64(&Ss), Ds);
+	LaneUserData->AddExponentialHeightFog(uint64(&Ss), Ds);
 }
 
 void FExponentialHeightFogStateStreamImpl::Render_OnUpdate(const FExponentialHeightFogStaticState& Ss, const FExponentialHeightFogDynamicState& Ds, FExponentialHeightFogSceneProxy*& UserData)
@@ -22,7 +17,7 @@ void FExponentialHeightFogStateStreamImpl::Render_OnUpdate(const FExponentialHei
 
 void FExponentialHeightFogStateStreamImpl::Render_OnDestroy(const FExponentialHeightFogStaticState& Ss, const FExponentialHeightFogDynamicState& Ds, FExponentialHeightFogSceneProxy*& UserData)
 {
-	Scene.RemoveExponentialHeightFog(uint64(&Ss));
+	LaneUserData->RemoveExponentialHeightFog(uint64(&Ss));
 }
 
 ////////////////////////////////////////////////////////////////////////////////////////////////////
diff --git a/Engine/Source/Runtime/Renderer/Private/StateStream/ExponentialHeightFogStateStreamImpl.h b/Engine/Source/Runtime/Renderer/Private/StateStream/ExponentialHeightFogStateStreamImpl.h
--- a/Engine/Source/Runtime/Renderer/Private/StateStream/ExponentialHeightFogStateStreamImpl.h
+++ b/Engine/Source/Runtime/Renderer/Private/StateStream/ExponentialHeightFogStateStreamImpl.h
@@ -12,7 +12,7 @@ class FExponentialHeightFogSceneProxy;
 
 ////////////////////////////////////////////////////////////////////////////////////////////////////
 
-struct FExponentialHeightFogStateStreamSettings : TStateStreamSettings<IExponentialHeightFogStateStream, FExponentialHeightFogSceneProxy>
+struct FExponentialHeightFogStateStreamSettings : TStateStreamSettings<IExponentialHeightFogStateStream, FExponentialHeightFogSceneProxy, FSceneInterface>
 {
 	static inline constexpr bool SkipCreatingDeletes = true;
 };
@@ -21,14 +21,9 @@ struct FExponentialHeightFogStateStreamSettings : TStateStreamSettings<IExponent
 
 class FExponentialHeightFogStateStreamImpl : public TStateStream<FExponentialHeightFogStateStreamSettings>
 {
-public:
-	FExponentialHeightFogStateStreamImpl(FSceneInterface& InScene);
-private:
 	UE_API virtual void Render_OnCreate(const FExponentialHeightFogStaticState& Ss, const FExponentialHeightFogDynamicState& Ds, FExponentialHeightFogSceneProxy*& UserData, bool IsDestroyedInSameFrame) override;
 	UE_API virtual void Render_OnUpdate(const FExponentialHeightFogStaticState& Ss, const FExponentialHeightFogDynamicState& Ds, FExponentialHeightFogSceneProxy*& UserData) override;
 	UE_API virtual void Render_OnDestroy(const FExponentialHeightFogStaticState& Ss, const FExponentialHeightFogDynamicState& Ds, FExponentialHeightFogSceneProxy*& UserData) override;
-
-	FSceneInterface& Scene;
 };
 
 ////////////////////////////////////////////////////////////////////////////////////////////////////
diff --git a/Engine/Source/Runtime/Renderer/Private/StateStream/ParticleSystemStateStreamImpl.cpp b/Engine/Source/Runtime/Renderer/Private/StateStream/ParticleSystemStateStreamImpl.cpp
--- a/Engine/Source/Runtime/Renderer/Private/StateStream/ParticleSystemStateStreamImpl.cpp
+++ b/Engine/Source/Runtime/Renderer/Private/StateStream/ParticleSystemStateStreamImpl.cpp
@@ -292,8 +292,7 @@ TArrayView<FParticleEventKismetData> FParticleSystemObjectCascade::GetKismetEven
 
 ////////////////////////////////////////////////////////////////////////////////////////////////////
 
-FParticleSystemStateStreamCascade::FParticleSystemStateStreamCascade(FSceneInterface& InScene)
-:	Scene(InScene)
+FParticleSystemStateStreamCascade::FParticleSystemStateStreamCascade()
 {
 }
 
@@ -333,6 +332,7 @@ void FParticleSystemStateStreamCascade::Render_OnCreate(const FParticleSystemSta
 	//{
 	//	LocalBounds = Mesh->GetBounds();
 	//}	
+	FSceneInterface& Scene = *LaneUserData;
 
 	Object.PrimitiveSceneDesc.RenderMatrix = Transform.ToMatrixWithScale();
 	Object.PrimitiveSceneDesc.AttachmentRootPosition = Transform.GetLocation();
@@ -413,7 +413,7 @@ STATESTREAM_CREATOR_INSTANCE_WITH_FUNC(FParticleSystemStateStreamCascade, [](con
 
 ////////////////////////////////////////////////////////////////////////////////////////////////////
 
-FParticleSystemHandle FParticleSystemStateStreamImpl::Game_CreateInstance(const FParticleSystemStaticState& Ss, const FParticleSystemDynamicState& Ds)
+FParticleSystemHandle FParticleSystemStateStreamImpl::Game_CreateInstance(const FParticleSystemStaticState& Ss, const FParticleSystemDynamicState& Ds, uint32 LaneId)
 {
 	if (Cast<UParticleSystem>(Ds.GetSystemAsset()))
 		return CascadeBackend->Game_CreateInstance(Ss, Ds);
diff --git a/Engine/Source/Runtime/Renderer/Private/StateStream/ParticleSystemStateStreamImpl.h b/Engine/Source/Runtime/Renderer/Private/StateStream/ParticleSystemStateStreamImpl.h
--- a/Engine/Source/Runtime/Renderer/Private/StateStream/ParticleSystemStateStreamImpl.h
+++ b/Engine/Source/Runtime/Renderer/Private/StateStream/ParticleSystemStateStreamImpl.h
@@ -106,7 +106,7 @@ inline constexpr uint32 ParticleSystemStateStreamCascadeId = 5;
 
 ////////////////////////////////////////////////////////////////////////////////////////////////////
 
-struct FParticleSystemStateStreamSettingsCascade : TStateStreamSettings<IParticleSystemStateStream, FParticleSystemObjectCascade>
+struct FParticleSystemStateStreamSettingsCascade : TStateStreamSettings<IParticleSystemStateStream, FParticleSystemObjectCascade, FSceneInterface>
 {
 	enum { Id = ParticleSystemStateStreamCascadeId };
 	static inline constexpr const TCHAR* DebugName = TEXT("ParticleSystem(Cascade)");
@@ -117,7 +117,7 @@ struct FParticleSystemStateStreamSettingsCascade : TStateStreamSettings<IParticl
 class FParticleSystemStateStreamCascade final : public TStateStream<FParticleSystemStateStreamSettingsCascade>
 {
 public:
-	FParticleSystemStateStreamCascade(FSceneInterface& InScene);
+	FParticleSystemStateStreamCascade();
 private:
 	void SetTransformObject(FParticleSystemObjectCascade& Object, const FParticleSystemDynamicState& Ds);
 	UE_API virtual void Render_OnCreate(const FParticleSystemStaticState& Ss, const FParticleSystemDynamicState& Ds, FParticleSystemObjectCascade*& UserData, bool IsDestroyedInSameFrame) override;
@@ -126,7 +126,6 @@ private:
 
 	virtual void Render_PostUpdate() override;
 
-	FSceneInterface& Scene;
 	TArray<FParticleSystemObjectCascade*> Objects;
 };
 
@@ -135,13 +134,14 @@ private:
 class FParticleSystemStateStreamImpl final : public IParticleSystemStateStream, public IStateStream
 {
 public:
-	FParticleSystemStateStreamImpl(FSceneInterface& InScene) {}
+	FParticleSystemStateStreamImpl() {}
 
-	virtual FParticleSystemHandle Game_CreateInstance(const FParticleSystemStaticState& Ss, const FParticleSystemDynamicState& Ds) override;
+	virtual FParticleSystemHandle Game_CreateInstance(const FParticleSystemStaticState& Ss, const FParticleSystemDynamicState& Ds, uint32 LaneId = DefaultLaneId) override;
 	virtual void SetOtherBackend(IParticleSystemStateStream* Other) override;
 
-	virtual void Game_BeginTick() override {}
-	virtual void Game_EndTick(StateStreamTime AbsoluteTime) override {}
+	virtual void Game_BeginTick(uint32 LaneId) override {}
+	virtual void Game_EndTick(StateStreamTime AbsoluteTime, uint32 LaneId) override {}
+	virtual void Game_SetDefaultLane(uint32 LaneId) override {}
 	virtual void Game_Exit() override {}
 	virtual void* Game_GetVoidPointer() override { return this; }
 	virtual void Render_Update(StateStreamTime AbsoluteTime) override {}
diff --git a/Engine/Source/Runtime/Renderer/Private/StateStream/SkinnedMeshStateStreamImpl.cpp b/Engine/Source/Runtime/Renderer/Private/StateStream/SkinnedMeshStateStreamImpl.cpp
--- a/Engine/Source/Runtime/Renderer/Private/StateStream/SkinnedMeshStateStreamImpl.cpp
+++ b/Engine/Source/Runtime/Renderer/Private/StateStream/SkinnedMeshStateStreamImpl.cpp
@@ -76,8 +76,7 @@ void FSkinnedMeshObject::OnTransformObjectDirty()
 
 ////////////////////////////////////////////////////////////////////////////////////////////////////
 
-FSkinnedMeshStateStreamImpl::FSkinnedMeshStateStreamImpl(FSceneInterface& InScene)
-:	Scene(InScene)
+FSkinnedMeshStateStreamImpl::FSkinnedMeshStateStreamImpl()
 {
 }
 
@@ -110,6 +109,8 @@ void FSkinnedMeshStateStreamImpl::Render_OnCreate(const FSkinnedMeshStaticState&
 		return;
 	}
 
+	FSceneInterface& Scene = *LaneUserData;
+
 	FSkinnedMeshObject& Object = *new FSkinnedMeshObject();
 	Object.AddRef();
 	Object.SkinnedAsset = SkinnedAsset;
@@ -212,6 +213,8 @@ void FSkinnedMeshStateStreamImpl::Render_PostUpdate()
 {
 	TStateStream<FSkinnedMeshStateStreamSettings>::Render_PostUpdate();
 
+	FSceneInterface& Scene = *LaneUserData;
+
 	// TODO: Get rid of this entirely.. no need to create these queues of stuff when it can be done directly render side
 	Scene.GetSkeletalMeshUpdater()->PushChangesToStream();
 }
diff --git a/Engine/Source/Runtime/Renderer/Private/StateStream/SkinnedMeshStateStreamImpl.h b/Engine/Source/Runtime/Renderer/Private/StateStream/SkinnedMeshStateStreamImpl.h
--- a/Engine/Source/Runtime/Renderer/Private/StateStream/SkinnedMeshStateStreamImpl.h
+++ b/Engine/Source/Runtime/Renderer/Private/StateStream/SkinnedMeshStateStreamImpl.h
@@ -46,7 +46,7 @@ private:
 
 ////////////////////////////////////////////////////////////////////////////////////////////////////
 
-struct FSkinnedMeshStateStreamSettings : TStateStreamSettings<ISkinnedMeshStateStream, FSkinnedMeshObject>
+struct FSkinnedMeshStateStreamSettings : TStateStreamSettings<ISkinnedMeshStateStream, FSkinnedMeshObject, FSceneInterface>
 {
 	static inline constexpr bool SkipCreatingDeletes = true;
 };
@@ -57,7 +57,7 @@ struct FSkinnedMeshStateStreamSettings : TStateStreamSettings<ISkinnedMeshStateS
 class FSkinnedMeshStateStreamImpl : public TStateStream<FSkinnedMeshStateStreamSettings>
 {
 public:
-	FSkinnedMeshStateStreamImpl(FSceneInterface& InScene);
+	FSkinnedMeshStateStreamImpl();
 private:
 	void SetTransformObject(FSkinnedMeshObject& Object, const FSkinnedMeshDynamicState& Ds);
 
@@ -66,8 +66,6 @@ private:
 	UE_API virtual void Render_OnCreate(const FSkinnedMeshStaticState& Ss, const FSkinnedMeshDynamicState& Ds, FSkinnedMeshObject*& UserData, bool IsDestroyedInSameFrame) override;
 	UE_API virtual void Render_OnUpdate(const FSkinnedMeshStaticState& Ss, const FSkinnedMeshDynamicState& Ds, FSkinnedMeshObject*& UserData) override;
 	UE_API virtual void Render_OnDestroy(const FSkinnedMeshStaticState& Ss, const FSkinnedMeshDynamicState& Ds, FSkinnedMeshObject*& UserData) override;
-
-	FSceneInterface& Scene;
 };
 
 ////////////////////////////////////////////////////////////////////////////////////////////////////
diff --git a/Engine/Source/Runtime/Renderer/Private/StateStream/SkyAtmosphereStateStreamImpl.cpp b/Engine/Source/Runtime/Renderer/Private/StateStream/SkyAtmosphereStateStreamImpl.cpp
--- a/Engine/Source/Runtime/Renderer/Private/StateStream/SkyAtmosphereStateStreamImpl.cpp
+++ b/Engine/Source/Runtime/Renderer/Private/StateStream/SkyAtmosphereStateStreamImpl.cpp
@@ -7,19 +7,22 @@
 
 ////////////////////////////////////////////////////////////////////////////////////////////////////
 
-FSkyAtmosphereStateStreamImpl::FSkyAtmosphereStateStreamImpl(FSceneInterface& InScene)
-:	Scene(InScene)
+FSkyAtmosphereStateStreamImpl::FSkyAtmosphereStateStreamImpl()
 {
 }
 
 void FSkyAtmosphereStateStreamImpl::Render_OnCreate(const FSkyAtmosphereStaticState& Ss, const FSkyAtmosphereDynamicState& Ds, FSkyAtmosphereSceneProxy*& UserData, bool IsDestroyedInSameFrame)
 {
+	FSceneInterface& Scene = *LaneUserData;
+
 	UserData = new FSkyAtmosphereSceneProxy(Ds);
 	Scene.AddSkyAtmosphere(UserData, Ds.GetBuilt());
 }
 
 void FSkyAtmosphereStateStreamImpl::Render_OnUpdate(const FSkyAtmosphereStaticState& Ss, const FSkyAtmosphereDynamicState& Ds, FSkyAtmosphereSceneProxy*& UserData)
 {
+	FSceneInterface& Scene = *LaneUserData;
+
 	// For some reason we need to recreate the entire sceneproxy every time
 #if 0
 	UserData->AtmosphereSetup = FAtmosphereSetup(Ds);
@@ -34,6 +37,7 @@ void FSkyAtmosphereStateStreamImpl::Render_OnUpdate(const FSkyAtmosphereStaticSt
 
 void FSkyAtmosphereStateStreamImpl::Render_OnDestroy(const FSkyAtmosphereStaticState& Ss, const FSkyAtmosphereDynamicState& Ds, FSkyAtmosphereSceneProxy*& UserData)
 {
+	FSceneInterface& Scene = *LaneUserData;
 	Scene.RemoveSkyAtmosphere(UserData);
 	delete UserData;
 	UserData = nullptr;
diff --git a/Engine/Source/Runtime/Renderer/Private/StateStream/SkyAtmosphereStateStreamImpl.h b/Engine/Source/Runtime/Renderer/Private/StateStream/SkyAtmosphereStateStreamImpl.h
--- a/Engine/Source/Runtime/Renderer/Private/StateStream/SkyAtmosphereStateStreamImpl.h
+++ b/Engine/Source/Runtime/Renderer/Private/StateStream/SkyAtmosphereStateStreamImpl.h
@@ -12,7 +12,7 @@ class FSkyAtmosphereSceneProxy;
 
 ////////////////////////////////////////////////////////////////////////////////////////////////////
 
-struct FSkyAtmosphereStateStreamSettings : TStateStreamSettings<ISkyAtmosphereStateStream, FSkyAtmosphereSceneProxy>
+struct FSkyAtmosphereStateStreamSettings : TStateStreamSettings<ISkyAtmosphereStateStream, FSkyAtmosphereSceneProxy, FSceneInterface>
 {
 	static inline constexpr bool SkipCreatingDeletes = true;
 };
@@ -22,13 +22,11 @@ struct FSkyAtmosphereStateStreamSettings : TStateStreamSettings<ISkyAtmosphereSt
 class FSkyAtmosphereStateStreamImpl : public TStateStream<FSkyAtmosphereStateStreamSettings>
 {
 public:
-	FSkyAtmosphereStateStreamImpl(FSceneInterface& InScene);
+	FSkyAtmosphereStateStreamImpl();
 private:
 	UE_API virtual void Render_OnCreate(const FSkyAtmosphereStaticState& Ss, const FSkyAtmosphereDynamicState& Ds, FSkyAtmosphereSceneProxy*& UserData, bool IsDestroyedInSameFrame) override;
 	UE_API virtual void Render_OnUpdate(const FSkyAtmosphereStaticState& Ss, const FSkyAtmosphereDynamicState& Ds, FSkyAtmosphereSceneProxy*& UserData) override;
 	UE_API virtual void Render_OnDestroy(const FSkyAtmosphereStaticState& Ss, const FSkyAtmosphereDynamicState& Ds, FSkyAtmosphereSceneProxy*& UserData) override;
-
-	FSceneInterface& Scene;
 };
 
 ////////////////////////////////////////////////////////////////////////////////////////////////////
diff --git a/Engine/Source/Runtime/Renderer/Private/StateStream/StaticMeshStateStreamImpl.cpp b/Engine/Source/Runtime/Renderer/Private/StateStream/StaticMeshStateStreamImpl.cpp
--- a/Engine/Source/Runtime/Renderer/Private/StateStream/StaticMeshStateStreamImpl.cpp
+++ b/Engine/Source/Runtime/Renderer/Private/StateStream/StaticMeshStateStreamImpl.cpp
@@ -29,8 +29,7 @@ void FStaticMeshObject::OnTransformObjectDirty()
 
 ////////////////////////////////////////////////////////////////////////////////////////////////////
 
-FStaticMeshStateStreamImpl::FStaticMeshStateStreamImpl(FSceneInterface& InScene)
-:	Scene(InScene)
+FStaticMeshStateStreamImpl::FStaticMeshStateStreamImpl()
 {
 }
 
@@ -62,6 +61,8 @@ void FStaticMeshStateStreamImpl::Render_OnCreate(const FStaticMeshStaticState& S
 		return;
 	}
 
+	FSceneInterface& Scene = *LaneUserData;
+
 	FStaticMeshObject& Object = *new FStaticMeshObject();
 	Object.AddRef();
 
diff --git a/Engine/Source/Runtime/Renderer/Private/StateStream/StaticMeshStateStreamImpl.h b/Engine/Source/Runtime/Renderer/Private/StateStream/StaticMeshStateStreamImpl.h
--- a/Engine/Source/Runtime/Renderer/Private/StateStream/StaticMeshStateStreamImpl.h
+++ b/Engine/Source/Runtime/Renderer/Private/StateStream/StaticMeshStateStreamImpl.h
@@ -34,7 +34,7 @@ private:
 
 ////////////////////////////////////////////////////////////////////////////////////////////////////
 
-struct FStaticMeshStateStreamSettings : TStateStreamSettings<IStaticMeshStateStream, FStaticMeshObject>
+struct FStaticMeshStateStreamSettings : TStateStreamSettings<IStaticMeshStateStream, FStaticMeshObject, FSceneInterface>
 {
 	static inline constexpr bool SkipCreatingDeletes = true;
 };
@@ -44,14 +44,12 @@ struct FStaticMeshStateStreamSettings : TStateStreamSettings<IStaticMeshStateStr
 class FStaticMeshStateStreamImpl : public TStateStream<FStaticMeshStateStreamSettings>
 {
 public:
-	FStaticMeshStateStreamImpl(FSceneInterface& InScene);
+	FStaticMeshStateStreamImpl();
 private:
 	void SetTransformObject(FStaticMeshObject& Object, const FStaticMeshDynamicState& Ds);
 	UE_API virtual void Render_OnCreate(const FStaticMeshStaticState& Ss, const FStaticMeshDynamicState& Ds, FStaticMeshObject*& UserData, bool IsDestroyedInSameFrame) override;
 	UE_API virtual void Render_OnUpdate(const FStaticMeshStaticState& Ss, const FStaticMeshDynamicState& Ds, FStaticMeshObject*& UserData) override;
 	UE_API virtual void Render_OnDestroy(const FStaticMeshStaticState& Ss, const FStaticMeshDynamicState& Ds, FStaticMeshObject*& UserData) override;
-
-	FSceneInterface& Scene;
 };
 
 ////////////////////////////////////////////////////////////////////////////////////////////////////
diff --git a/Engine/Source/Runtime/StateStream/Private/StateStreamManagerImpl.cpp b/Engine/Source/Runtime/StateStream/Private/StateStreamManagerImpl.cpp
--- a/Engine/Source/Runtime/StateStream/Private/StateStreamManagerImpl.cpp
+++ b/Engine/Source/Runtime/StateStream/Private/StateStreamManagerImpl.cpp
@@ -6,74 +6,42 @@
 
 #define STATE_STREAM_DEBUG_CREATE_DESTROY 0
 
-void FStateStreamDebugElements::DrawDebugElements(const FString& DebugName, IStateStreamDebugRenderer& DebugRenderer) const
-{
-	TStringBuilder<1024> DebugLine;
-			
-	const int32 NumElementCounts = DebugElementsCounts.Num();
-	if (NumElementCounts == 0)
-	{
-		return;
-	}
-	const float NumElementColumns = FMath::Min(10.0f, NumElementCounts);
-	const int32 NumElementRows = FMath::CeilToInt32(NumElementCounts / NumElementColumns);
-	
-	// Display the number of elements
-	DebugLine.Reset();
-	DebugLine.Appendf(TEXT("Num %s = %d"), *DebugName, NumElementCounts);
-	DebugRenderer.DrawText(*DebugLine, FColor::Blue);
-	
-	// Display per element informations
-	auto ElementIterator = DebugElementsCounts.CreateConstIterator();
-	for (int32 ColumnIndex = 0; ColumnIndex < NumElementColumns; ++ColumnIndex)
-	{
-		DebugLine.Reset();
-		DebugLine.Appendf(TEXT("%s : "), *DebugName);
-		for (int32 RowIndex = 0; RowIndex < NumElementRows; ++RowIndex)
-		{
-			if (ElementIterator)
-			{
-				DebugLine.Appendf(TEXT("[%s | %d]"), *ElementIterator.Key(), ElementIterator.Value());
-				UE_LOG(LogTemp, Log, TEXT("%s : [%s | %d]"), *DebugName, *ElementIterator.Key(), ElementIterator.Value());
-				++ElementIterator;
-			}
-		}
-		DebugRenderer.DrawText(*DebugLine, FColor::Orange);
-	}
-}
-
 ////////////////////////////////////////////////////////////////////////////////////////////////////
 
-FStateStreamDebugElements FStateStreamManagerImpl::DebugRenderCommands;
-
-FStateStreamDebugElements& FStateStreamManagerImpl::GetDebugRenderCommands()
+uint32 FStateStreamManagerImpl::Game_CreateLane()
 {
-	return DebugRenderCommands;
+	return 0;
+}
+void FStateStreamManagerImpl::Game_SetLaneUserData(uint32 LaneId, void* UserData)
+{
+	for (const StateStreamRec& Rec : StateStreams)
+	{
+		Rec.Stream->Game_SetLaneUserData(LaneId, UserData);
+	}
 }
 
-FStateStreamDebugElements& FStateStreamManagerImpl::GetDebugSceneProxies()
+void FStateStreamManagerImpl::Game_DestroyLane(uint32 LaneId)
 {
-	return DebugSceneProxies;
 }
 
-void FStateStreamManagerImpl::Game_BeginTick()
+void FStateStreamManagerImpl::Game_BeginTick(uint32 LaneId)
 {
 	check(!bIsInTick);
 	check (!bGameExited);
 	bIsInTick = true;
 	for (const StateStreamRec& Rec : StateStreams)
 	{
-		Rec.Stream->Game_BeginTick();
+		Rec.Stream->Game_BeginTick(LaneId);
 	}
 }
 
-void FStateStreamManagerImpl::Game_EndTick(double AbsoluteTime)
+void FStateStreamManagerImpl::Game_EndTick(double AbsoluteTime, uint32 LaneId)
 {
 	check(bIsInTick);
 	bIsInTick = false;
 	for (const StateStreamRec& Rec : StateStreams)
 	{
-		Rec.Stream->Game_EndTick(AbsoluteTime);
+		Rec.Stream->Game_EndTick(AbsoluteTime, LaneId);
 	}
 }
 
@@ -85,18 +53,24 @@ void FStateStreamManagerImpl::Game_Exit()
 		Rec.Stream->Game_Exit();
 	}
 	bGameExited = true;
-	
-	// Reset all the scene proxies at the end of the game since 
-	// we are monitoring all the registered proxies during 
-	// the full game life time
-	DebugSceneProxies.ResetDebugElements();
 }
 
-bool FStateStreamManagerImpl::Game_IsInTick()
+bool FStateStreamManagerImpl::Game_IsInTick(uint32 LaneId)
 {
 	return bIsInTick;
 }
 
+uint32 FStateStreamManagerImpl::Game_SetDefaultLane(uint32 LaneId)
+{
+	uint32 OldLaneId = ActiveDefaultLaneId;
+	ActiveDefaultLaneId = LaneId;
+	for (const StateStreamRec& Rec : StateStreams)
+	{
+		Rec.Stream->Game_SetDefaultLane(LaneId);
+	}
+	return OldLaneId;
+}
+
 void* FStateStreamManagerImpl::Game_GetStreamPointer(uint32 Id)
 {
 	check(Id < uint32(StateStreamsLookup.Num()));
@@ -187,7 +161,7 @@ IStateStream* FStateStreamManagerImpl::Render_GetStream(uint32 Id)
 	return StateStreamsLookup[Id];
 }
 
-void FStateStreamManagerImpl::Game_DebugRender(IStateStreamDebugRenderer& Renderer)
+void FStateStreamManagerImpl::Game_DebugRender(IStateStreamDebugRenderer& Renderer, uint32 LaneId)
 {
 	TStringBuilder<1024> DebugLine;
 	
@@ -201,12 +175,6 @@ void FStateStreamManagerImpl::Game_DebugRender(IStateStreamDebugRenderer& Render
 	{
 		Rec.Stream->Game_DebugRender(Renderer);
 	}
-	// Draw the render commands and scene proxies debug data
-	DebugRenderCommands.DrawDebugElements(FString("Render command"), Renderer);
-	DebugSceneProxies.DrawDebugElements(FString("Scene proxies"), Renderer);
-	
-	// Reset the render commands debug elements every frame after using them
-	DebugRenderCommands.ResetDebugElements();
 }
 
 FStateStreamManagerImpl::~FStateStreamManagerImpl()
diff --git a/Engine/Source/Runtime/StateStream/Public/GenericStateStream.h b/Engine/Source/Runtime/StateStream/Public/GenericStateStream.h
--- a/Engine/Source/Runtime/StateStream/Public/GenericStateStream.h
+++ b/Engine/Source/Runtime/StateStream/Public/GenericStateStream.h
@@ -28,11 +28,12 @@ DECLARE_STATS_GROUP(TEXT("StateStream"), STATGROUP_StateStream, STATCAT_Advanced
 // Settings struct - Can be inherited to set settings on state stream. It mainly contains all the data types used 
 // by state stream (static state / dynamic state / handle type / interface type) + id and debug name
 
-template<typename TInterfaceType, typename TUserDataType = void>
+template<typename TInterfaceType, typename TInstanceUserDataType = void, typename TLaneUserDataType = void>
 struct TStateStreamSettings
 {
 	using InterfaceType = TInterfaceType;
-	using UserDataType = TUserDataType;
+	using InstanceUserDataType = TInstanceUserDataType;
+	using LaneUserDataType = TLaneUserDataType;
 	enum { Id = InterfaceType::Id };
 	static inline constexpr const TCHAR* DebugName = InterfaceType::Handle::DebugName;
 	static inline constexpr bool SkipCreatingDeletes = false;
@@ -52,11 +53,12 @@ public:
 	using FHandle = typename InterfaceType::Handle;
 	using FStaticState = typename InterfaceType::StaticState;
 	using FDynamicState = typename InterfaceType::DynamicState;
-	using FUserDataType = typename Settings::UserDataType;
+	using FInstanceUserDataType = typename Settings::InstanceUserDataType;
+	using FLaneUserDataType = typename Settings::LaneUserDataType;
 	enum { Id = Settings::Id };
 
 	//~ Begin InterfaceType interface
-	virtual FHandle Game_CreateInstance(const FStaticState& Ss, const FDynamicState& Ds) override final;
+	virtual FHandle Game_CreateInstance(const FStaticState& Ss, const FDynamicState& Ds, uint32 LaneId = DefaultLaneId) override final;
 	//~ End InterfaceType interface
 
 	//~ Begin IStateStreamHandleOwner interface (used by state stream handles on game side)
@@ -68,8 +70,12 @@ public:
 	//~ End IStateStreamHandleOwner interface
 	
 	//~ Begin IStateStream interface (used by StateStreamManagerImpl)
-	virtual void Game_BeginTick() override;
-	virtual void Game_EndTick(StateStreamTime AbsoluteTime) override;
+	virtual void Game_CreateLane() override {}
+	virtual void Game_SetLaneUserData(uint32 LaneId, void* UserData) override { LaneUserData = (FLaneUserDataType*)UserData; }
+	virtual void Game_DestroyLane(uint32 LaneId) override {}
+	virtual void Game_BeginTick(uint32 LaneId) override;
+	virtual void Game_EndTick(StateStreamTime AbsoluteTime, uint32 LaneId) override;
+	virtual void Game_SetDefaultLane(uint32 LaneId) override;
 	virtual void Game_Exit() override;
 	virtual void* Game_GetVoidPointer() override;
 	virtual void Game_DebugRender(IStateStreamDebugRenderer& Renderer) override;
@@ -83,20 +89,20 @@ public:
 	//~ End IStateStream interface
 
 	// Get the instance user data on the render side
-	FUserDataType*& Render_GetUserData(const FHandle& Handle);
+	FInstanceUserDataType*& Render_GetUserData(const FHandle& Handle);
 	
 	// Get the instance interpolated dynamic state on the render side
 	const FDynamicState& Render_GetDynamicState(const FHandle& Handle);
 	
 	// Specialize to do custom things when state has been created / updated or destroyed
-	virtual void Render_OnCreate(const FStaticState& Ss, const FDynamicState& Ds, FUserDataType*& UserData, bool IsDestroyedInSameFrame) {}
-	virtual void Render_OnUpdate(const FStaticState& Ss, const FDynamicState& Ds, FUserDataType*& UserData) {}
-	virtual void Render_OnDestroy(const FStaticState& Ss, const FDynamicState& Ds, FUserDataType*& UserData) {}
+	virtual void Render_OnCreate(const FStaticState& Ss, const FDynamicState& Ds, FInstanceUserDataType*& UserData, bool IsDestroyedInSameFrame) {}
+	virtual void Render_OnUpdate(const FStaticState& Ss, const FDynamicState& Ds, FInstanceUserDataType*& UserData) {}
+	virtual void Render_OnDestroy(const FStaticState& Ss, const FDynamicState& Ds, FInstanceUserDataType*& UserData) {}
 
 	// ... or use template specialization to avoid virtual calls
-	void Render_OnCreateInline(const FStaticState& Ss, const FDynamicState& Ds, FUserDataType*& UserData, bool IsDestroyedInSameFrame) { Render_OnCreate(Ss, Ds, UserData, IsDestroyedInSameFrame); }
-	void Render_OnUpdateInline(const FStaticState& Ss, const FDynamicState& Ds, FUserDataType*& UserData) { Render_OnUpdate(Ss, Ds, UserData); }
-	void Render_OnDestroyInline(const FStaticState& Ss, const FDynamicState& Ds, FUserDataType*& UserData) { Render_OnDestroy(Ss, Ds, UserData); }
+	void Render_OnCreateInline(const FStaticState& Ss, const FDynamicState& Ds, FInstanceUserDataType*& UserData, bool IsDestroyedInSameFrame) { Render_OnCreate(Ss, Ds, UserData, IsDestroyedInSameFrame); }
+	void Render_OnUpdateInline(const FStaticState& Ss, const FDynamicState& Ds, FInstanceUserDataType*& UserData) { Render_OnUpdate(Ss, Ds, UserData); }
+	void Render_OnDestroyInline(const FStaticState& Ss, const FDynamicState& Ds, FInstanceUserDataType*& UserData) { Render_OnDestroy(Ss, Ds, UserData); }
 	
 	// Get the used instance counts within the streams
 	uint32 GetUsedInstancesCount() const { return Instances.GetUsedCount(); }
@@ -119,7 +125,7 @@ protected:
 		uint32 CreateTick = 0;
 		uint32 DeleteTick = ~0u;
 		TOptional<FDynamicState> RendDynamicState;
-		FUserDataType* UserData = nullptr;
+		FInstanceUserDataType* UserData = nullptr;
 	};
 	
 	// Tick produced by game side.
@@ -147,6 +153,10 @@ protected:
 	template<typename T>
 	void MakeInternal(T& State);
 
+	uint32 ActiveDefaultLaneId = DefaultLaneId;
+
+	FLaneUserDataType* LaneUserData = nullptr;
+
 	// Lock in order to write into the current tick
 	FRWLock CurrentTickLock;
 	
@@ -190,8 +200,14 @@ protected:
 ////////////////////////////////////////////////////////////////////////////////////////////////////
 
 template<typename Settings>
-typename TStateStream<Settings>::FHandle TStateStream<Settings>::Game_CreateInstance(const FStaticState& Ss, const FDynamicState& Ds)
+typename TStateStream<Settings>::FHandle TStateStream<Settings>::Game_CreateInstance(const FStaticState& Ss, const FDynamicState& Ds, uint32 LaneId)
 {
+	if (LaneId == DefaultLaneId)
+	{
+		check(ActiveDefaultLaneId != DefaultLaneId);
+		LaneId = ActiveDefaultLaneId;
+	}
+
 	check(CurrentTick);
 
 	// We usually have several dynamic states stored for each instance
@@ -348,7 +364,7 @@ void* TStateStream<Settings>::Render_GetUserData(uint32 HandleId)
 }
 
 template<typename Settings>
-void TStateStream<Settings>::Game_BeginTick()
+void TStateStream<Settings>::Game_BeginTick(uint32 LaneId)
 {
 	check(!CurrentTick);
 	
@@ -365,7 +381,7 @@ void TStateStream<Settings>::Game_BeginTick()
 }
 
 template<typename Settings>
-void TStateStream<Settings>::Game_EndTick(StateStreamTime AbsoluteTime)
+void TStateStream<Settings>::Game_EndTick(StateStreamTime AbsoluteTime, uint32 LaneId)
 {
 	check(CurrentTick);
 
@@ -387,6 +403,12 @@ void TStateStream<Settings>::Game_EndTick(StateStreamTime AbsoluteTime)
 	CurrentTick = nullptr;
 }
 
+template<typename Settings>
+void TStateStream<Settings>::Game_SetDefaultLane(uint32 LaneId)
+{
+	ActiveDefaultLaneId = LaneId;
+}
+
 template<typename Settings>
 void TStateStream<Settings>::Game_Exit()
 {
@@ -673,7 +695,7 @@ void TStateStream<Settings>::Game_DebugRender(IStateStreamDebugRenderer& Rendere
 }
 
 template<typename Settings>
-typename TStateStream<Settings>::FUserDataType*& TStateStream<Settings>::Render_GetUserData(const FHandle& Handle)
+typename TStateStream<Settings>::FInstanceUserDataType*& TStateStream<Settings>::Render_GetUserData(const FHandle& Handle)
 {
 	uint32 Id = Handle.GetId();
 	check(Id);
diff --git a/Engine/Source/Runtime/StateStream/Public/StateStream.h b/Engine/Source/Runtime/StateStream/Public/StateStream.h
--- a/Engine/Source/Runtime/StateStream/Public/StateStream.h
+++ b/Engine/Source/Runtime/StateStream/Public/StateStream.h
@@ -13,12 +13,19 @@ class IStateStreamDebugRenderer;
 class IStateStream
 {
 public:
+	virtual void Game_CreateLane() {}
+	virtual void Game_SetLaneUserData(uint32 LaneId, void* UserData) {}
+	virtual void Game_DestroyLane(uint32 LaneId) {}
+
 	// Create current tick data that will be update later on 
-	virtual void Game_BeginTick() = 0;
+	virtual void Game_BeginTick(uint32 LaneId) = 0;
 	
 	// Make the current tick available to the renderer (update oldest/newest ticks) and then reset it 
-	virtual void Game_EndTick(StateStreamTime AbsoluteTime) = 0;
-	
+	virtual void Game_EndTick(StateStreamTime AbsoluteTime, uint32 LaneId) = 0;
+
+	// Set the lane that will be used if provided lane is DefaultLaneId
+	virtual void Game_SetDefaultLane(uint32 LaneId) = 0;
+
 	// Release all the instance/ticks/dynamic states stored within that stream
 	virtual void Game_Exit() = 0;
 	
diff --git a/Engine/Source/Runtime/StateStream/Public/StateStreamCreator.h b/Engine/Source/Runtime/StateStream/Public/StateStreamCreator.h
--- a/Engine/Source/Runtime/StateStream/Public/StateStreamCreator.h
+++ b/Engine/Source/Runtime/StateStream/Public/StateStreamCreator.h
@@ -6,7 +6,6 @@
 
 #define UE_API STATESTREAM_API
 
-class FSceneInterface;
 class FStateStreamManagerImpl;
 class IStateStream;
 
@@ -15,7 +14,6 @@ class IStateStream;
 struct FStateStreamRegisterContext
 {
 	FStateStreamManagerImpl& Manager;
-	FSceneInterface* Scene;
 
 	UE_API void Register(IStateStream& StateStream, bool TakeOwnership) const;
 	UE_API void RegisterDependency(uint32 FromId, uint32 ToId) const;
@@ -61,19 +59,19 @@ private:
 #define STATESTREAM_CREATOR_INSTANCE(ImplName) \
 	FStateStreamCreator ImplName##Creator(ImplName::Id, \
 		[](const FStateStreamRegisterContext& Context) \
-			{ ImplName& Impl = *new ImplName(*Context.Scene); Context.Register(Impl, true); }, \
+			{ ImplName& Impl = *new ImplName(); Context.Register(Impl, true); }, \
 		[](const FStateStreamUnregisterContext& Context) {});
 
 #define STATESTREAM_CREATOR_INSTANCE_WITH_DEPENDENCY(ImplName, ToId) \
 	FStateStreamCreator ImplName##Creator(ImplName::Id, \
 		[](const FStateStreamRegisterContext& Context) \
-			{ ImplName& Impl = *new ImplName(*Context.Scene); Context.Register(Impl, true); Context.RegisterDependency(ImplName::Id, ToId); }, \
+			{ ImplName& Impl = *new ImplName(); Context.Register(Impl, true); Context.RegisterDependency(ImplName::Id, ToId); }, \
 		[](const FStateStreamUnregisterContext& Context) {});
 
 #define STATESTREAM_CREATOR_INSTANCE_WITH_FUNC(ImplName, ...) \
 	FStateStreamCreator ImplName##Creator(ImplName::Id, \
 		[](const FStateStreamRegisterContext& Context) \
-			{ ImplName& Impl = *new ImplName(*Context.Scene); Context.Register(Impl, true); __VA_ARGS__(Context, Impl); }, \
+			{ ImplName& Impl = *new ImplName(); Context.Register(Impl, true); __VA_ARGS__(Context, Impl); }, \
 		[](const FStateStreamUnregisterContext& Context) {});
 
 ////////////////////////////////////////////////////////////////////////////////////////////////////
diff --git a/Engine/Source/Runtime/StateStream/Public/StateStreamDebugRenderer.h b/Engine/Source/Runtime/StateStream/Public/StateStreamDebugRenderer.h
--- a/Engine/Source/Runtime/StateStream/Public/StateStreamDebugRenderer.h
+++ b/Engine/Source/Runtime/StateStream/Public/StateStreamDebugRenderer.h
@@ -13,6 +13,7 @@ class IStateStreamDebugRenderer
 public:
 	// Interface to draw text with a given color for debug purposes
 	virtual void DrawText(const FStringView& Text, const FColor& Color) = 0;
+	virtual void DrawText(uint32 X, uint32 Y, const FStringView& Text, const FColor& Color) = 0;
 	virtual ~IStateStreamDebugRenderer() {}
 };
 
diff --git a/Engine/Source/Runtime/StateStream/Public/StateStreamDefinitions.h b/Engine/Source/Runtime/StateStream/Public/StateStreamDefinitions.h
--- a/Engine/Source/Runtime/StateStream/Public/StateStreamDefinitions.h
+++ b/Engine/Source/Runtime/StateStream/Public/StateStreamDefinitions.h
@@ -4,6 +4,11 @@
 
 #include "StateStreamHandle.h"
 
+
+////////////////////////////////////////////////////////////////////////////////////////////////////
+
+inline constexpr uint32 DefaultLaneId = ~0u;
+
 ////////////////////////////////////////////////////////////////////////////////////////////////////
 
 #define DECLARE_STATESTREAM(Type) \
diff --git a/Engine/Source/Runtime/StateStream/Public/StateStreamManager.h b/Engine/Source/Runtime/StateStream/Public/StateStreamManager.h
--- a/Engine/Source/Runtime/StateStream/Public/StateStreamManager.h
+++ b/Engine/Source/Runtime/StateStream/Public/StateStreamManager.h
@@ -12,19 +12,26 @@ class IStateStreamDebugRenderer;
 class IStateStreamManager
 {
 public:
+	virtual uint32 Game_CreateLane() { return 0; }
+	virtual void Game_SetLaneUserData(uint32 LaneId, void* UserData) {}
+	virtual void Game_DestroyLane(uint32 LaneId) {}
+
 	// Call from Game when a new tick is opened
 	// Note, no state stream handles can be created,updated,destroyed outside a Begin/End tick
-	virtual void Game_BeginTick() = 0;
+	virtual void Game_BeginTick(uint32 LaneId) = 0;
 
 	// Close tick and make it available to render side.
 	// AbsoluteTime is the amount of time that Game consumed
-	virtual void Game_EndTick(double AbsoluteTime) = 0;
+	virtual void Game_EndTick(double AbsoluteTime, uint32 LaneId) = 0;
 
 	// Should be called when game is exiting.
 	virtual void Game_Exit() = 0;
 
 	// Returns true if game is inside an open tick
-	virtual bool Game_IsInTick() = 0;
+	virtual bool Game_IsInTick(uint32 LaneId) = 0;
+
+	// Set the lane that will be used if provided lane is DefaultLaneId. Returns previous default lane id
+	virtual uint32 Game_SetDefaultLane(uint32 LaneId) = 0;
 
 	// Functions to fetch StateStream interface (not IStateStream) game side.
 	template<typename T>
@@ -32,10 +39,33 @@ public:
 	virtual void* Game_GetStreamPointer(uint32 Id) = 0;
 
 	// StateStream debug rendering
-	virtual void Game_DebugRender(IStateStreamDebugRenderer& Renderer) = 0;
+	virtual void Game_DebugRender(IStateStreamDebugRenderer& Renderer, uint32 LaneId) = 0;
 
 protected:
 	virtual ~IStateStreamManager() = default;
 };
 
 ////////////////////////////////////////////////////////////////////////////////////////////////////
+
+struct FStateStreamLaneScope
+{
+	FStateStreamLaneScope(IStateStreamManager* InManager, uint32 LaneId) : Manager(InManager)
+	{
+		if (Manager)
+		{
+			OldLaneId = Manager->Game_SetDefaultLane(LaneId);
+		}
+	}
+	~FStateStreamLaneScope()
+	{
+		if (Manager)
+		{
+			Manager->Game_SetDefaultLane(OldLaneId);
+		}
+	}
+
+	IStateStreamManager* Manager;
+	uint32 OldLaneId;
+};
+
+////////////////////////////////////////////////////////////////////////////////////////////////////
diff --git a/Engine/Source/Runtime/StateStream/Public/StateStreamManagerImpl.h b/Engine/Source/Runtime/StateStream/Public/StateStreamManagerImpl.h
--- a/Engine/Source/Runtime/StateStream/Public/StateStreamManagerImpl.h
+++ b/Engine/Source/Runtime/StateStream/Public/StateStreamManagerImpl.h
@@ -8,44 +8,6 @@
 
 #define UE_API STATESTREAM_API
 
-////////////////////////////////////////////////////////////////////////////////////////////////////
-// State stream debug elements to get information
-// about what remains to convert to state stream 
-
-class FStateStreamDebugElements
-{
-public:
-	// Add N debug elements given a type 
-	void AddDebugElements(const FString& ElementType, const uint32 NumElements)
-	{
-		if (uint32* ElementCount = DebugElementsCounts.Find(ElementType))
-		{
-			*ElementCount += NumElements;
-		}
-		else
-		{
-			DebugElementsCounts.Add(ElementType, NumElements);
-		}
-	}
-	
-	// Reset all the debug elements 
-	void ResetDebugElements()
-	{
-		DebugElementsCounts.Reset();
-	}
-	
-	// Draw the debug elements to the renderer 
-	void DrawDebugElements(const FString& DebubName, IStateStreamDebugRenderer& DebugRenderer) const;
-	
-	// Get the debug elements counts 
-	const TMap<FString,uint32>& GetDebugElements() const { return DebugElementsCounts; }
-
-private : 
-	
-	// Number of debug elements per type 
-	TMap<FString,uint32> DebugElementsCounts;
-};
-
 ////////////////////////////////////////////////////////////////////////////////////////////////////
 // StateStreamManager implementation. 
 // This type should only be known render side.
@@ -54,12 +16,16 @@ class FStateStreamManagerImpl : public IStateStreamManager
 {
 public:
 	//~ Begin IStateStreamManager interface
-	UE_API virtual void Game_BeginTick() override;
-	UE_API virtual void Game_EndTick(double AbsoluteTime) override;
+	UE_API virtual uint32 Game_CreateLane();
+	UE_API virtual void Game_SetLaneUserData(uint32 LaneId, void* UserData);
+	UE_API virtual void Game_DestroyLane(uint32 LaneId);
+	UE_API virtual void Game_BeginTick(uint32 LaneId) override;
+	UE_API virtual void Game_EndTick(double AbsoluteTime, uint32 LaneId) override;
 	UE_API virtual void Game_Exit() override;
-	UE_API virtual bool Game_IsInTick() override;
+	UE_API virtual bool Game_IsInTick(uint32 LaneId) override;
+	UE_API virtual uint32 Game_SetDefaultLane(uint32 LaneId) override;
 	UE_API virtual void* Game_GetStreamPointer(uint32 Id) override;
-	UE_API virtual void Game_DebugRender(IStateStreamDebugRenderer& Renderer) override;
+	UE_API virtual void Game_DebugRender(IStateStreamDebugRenderer& Renderer, uint32 LaneId) override;
 	//~ End IStateStreamManager interface
 
 	// Register new state streams into manager.
@@ -84,12 +50,6 @@ public:
 	// Get state stream from id
 	UE_API IStateStream* Render_GetStream(uint32 Id);
 	
-	// Get the debug render commands 
-	UE_API static FStateStreamDebugElements& GetDebugRenderCommands();
-	
-	// Get the debug scene proxies 
-	UE_API FStateStreamDebugElements& GetDebugSceneProxies();
-	
 	UE_API virtual ~FStateStreamManagerImpl() override;
 	
 private:
@@ -107,7 +67,9 @@ private:
 	// List of state streams for each registered state stream id. 
 	// This will allow fast lookup to retrieve a state stream given an id
 	TArray<IStateStream*> StateStreamsLookup;
-	
+
+	uint32 ActiveDefaultLaneId = DefaultLaneId;
+
 	// Bool to check if we are currently in between the begin and the end tick
 	bool bIsInTick = false;
 	
@@ -119,12 +81,6 @@ private:
 
 	// Async task to perform the garbage collection
 	UE::Tasks::FTask GarbageCollectTask;
-	
-	// Render commands debug data to keep track of the remaining commands to be converted to state stream
-	static FStateStreamDebugElements DebugRenderCommands;
-	
-	// Scene proxies debug data to keep track of the remaining proxies to be converted to state stream
-	FStateStreamDebugElements DebugSceneProxies;
 };
 
 ////////////////////////////////////////////////////////////////////////////////////////////////////
diff --git a/Engine/Source/Runtime/StateStream/Public/TransformStateStream.h b/Engine/Source/Runtime/StateStream/Public/TransformStateStream.h
--- a/Engine/Source/Runtime/StateStream/Public/TransformStateStream.h
+++ b/Engine/Source/Runtime/StateStream/Public/TransformStateStream.h
@@ -65,7 +65,7 @@ class ITransformStateStream
 {
 public:
 	DECLARE_STATESTREAM(Transform)
-	virtual FTransformHandle Game_CreateInstance(const FTransformStaticState& Ss, const FTransformDynamicState& Ds) = 0;
+	virtual FTransformHandle Game_CreateInstance(const FTransformStaticState& Ss, const FTransformDynamicState& Ds, uint32 LaneId = DefaultLaneId) = 0;
 };
 
 ////////////////////////////////////////////////////////////////////////////////////////////////////
