[CoreUObject/UHT]
* Added support for UPARTIAL in unreal.

#Synced-CL 51378911

diff --git a/Engine/Source/Programs/LowLevelTests/FoundationTests/Private/UPartialTestClasses.h b/Engine/Source/Programs/LowLevelTests/FoundationTests/Private/UPartialTestClasses.h
new file mode 100644
--- /dev/null
+++ b/Engine/Source/Programs/LowLevelTests/FoundationTests/Private/UPartialTestClasses.h
@@ -0,0 +1,85 @@
+// Copyright Epic Games, Inc. All Rights Reserved.
+
+#pragma once
+
+#include "CoreMinimal.h"
+#include "Templates/SubclassOf.h"
+#include "UObject/Object.h"
+#include "UObject/UObjectPartials.h"
+#include "UObject/Class.h"
+#include "UPartialTestClasses.generated.h"
+
+UCLASS(Partial)
+class UTestPartialObject : public UObject
+{
+	GENERATED_BODY()
+public:
+	UPROPERTY()
+	int32 NativeValue = 100;
+
+	static void AddReferencedObjects(UObject* InThis, FReferenceCollector& Collector)
+	{
+		// TODO: We want to support AddReferencedObjects on Partials
+		// solution should be to add the Partials ARO directly to the schema
+		// when there are Partials found that has the function
+		Super::AddReferencedObjects(InThis, Collector);
+	}
+
+
+	UFUNCTION(BlueprintPure, Category = "PartialTest")
+	int32 GetFooooo() const
+	{
+		return 32;
+	}
+};
+
+UCLASS(Partial)
+class UDerivedTestPartialObject : public UTestPartialObject
+{
+	GENERATED_BODY()
+public:
+	UPROPERTY()
+	int32 DerivedValue = 200;
+};
+
+UCLASS() // On purpose skip Partial attribute
+class UDerivedDerivedTestPartialObject : public UDerivedTestPartialObject
+{
+	GENERATED_BODY()
+};
+
+UCLASS(Blueprintable, Partial)
+class UBlueprintablePartialObject : public UObject
+{
+	GENERATED_BODY()
+public:
+	UPROPERTY(BlueprintReadWrite, Category = "PartialTest")
+	int32 BlueprintValue = 500;
+};
+
+UCLASS(Partial)
+class UPartialTestObject : public UObject
+{
+	GENERATED_BODY()
+};
+
+UCLASS(Partial)
+class USerializationTestObject : public UObject
+{
+	GENERATED_BODY()
+public:
+	UPROPERTY()
+	int32 NativeProperty = 100;
+
+	UPROPERTY()
+	FString NativeString = TEXT("Native");
+};
+
+UCLASS(Partial)
+class UDerivedSerializationTestObject : public USerializationTestObject
+{
+	GENERATED_BODY()
+public:
+	UPROPERTY()
+	int32 DerivedProperty = 200;
+};
diff --git a/Engine/Source/Programs/LowLevelTests/FoundationTests/Private/UPartialTestOutlinedFunctions.cpp b/Engine/Source/Programs/LowLevelTests/FoundationTests/Private/UPartialTestOutlinedFunctions.cpp
new file mode 100644
--- /dev/null
+++ b/Engine/Source/Programs/LowLevelTests/FoundationTests/Private/UPartialTestOutlinedFunctions.cpp
@@ -0,0 +1,28 @@
+// Copyright Epic Games, Inc. All Rights Reserved.
+
+#include "UPartialTestPartials.h"
+
+// Implementation of FTestOutlinedFunctionPartial functions
+// These are outlined (not inline) to test that UHT properly generates
+// function registration for non-inline partial functions
+
+int32 FTestOutlinedFunctionPartial::GetStoredValue() const
+{
+	return StoredValue;
+}
+
+void FTestOutlinedFunctionPartial::SetStoredValue(int32 NewValue)
+{
+	StoredValue = NewValue;
+}
+
+int32 FTestOutlinedFunctionPartial::MultiplyStoredValue(int32 Multiplier)
+{
+	StoredValue *= Multiplier;
+	return StoredValue;
+}
+
+int32 FTestOutlinedFunctionPartial::GetOwnerNativeValue() const
+{
+	return GetOwner().NativeValue;
+}
diff --git a/Engine/Source/Programs/LowLevelTests/FoundationTests/Private/UPartialTestPartials.h b/Engine/Source/Programs/LowLevelTests/FoundationTests/Private/UPartialTestPartials.h
new file mode 100644
--- /dev/null
+++ b/Engine/Source/Programs/LowLevelTests/FoundationTests/Private/UPartialTestPartials.h
@@ -0,0 +1,256 @@
+// Copyright Epic Games, Inc. All Rights Reserved.
+
+#pragma once
+
+#include "UPartialTestClasses.h"
+#include "UPartialTestPartials.generated.h"
+
+UPARTIAL(UTestPartialObject)
+struct FTestObjectPartial
+{
+	GENERATED_BODY()
+
+	UPROPERTY()
+	int32 PartialValue = 43;
+
+	UPROPERTY()
+	FString PartialString = TEXT("PartialDefault");
+
+	UPROPERTY()
+	TArray<int32> PartialArray;
+
+	FTestObjectPartial()
+	{
+		PartialArray.Add(1);
+		PartialArray.Add(2);
+		PartialArray.Add(3);
+	}
+
+	void BeginDestroy()
+	{
+		PartialArray.Empty();
+	}
+};
+
+UPARTIAL(UTestPartialObject)
+struct FSecondObjectPartial
+{
+	GENERATED_BODY()
+
+	UPROPERTY()
+	float PartialFloat = 3.14f;
+
+	UPROPERTY()
+	bool bPartialBool = true;
+};
+
+UPARTIAL(UDerivedTestPartialObject)
+struct FDerivedObjectPartial
+{
+	GENERATED_BODY()
+
+	UPROPERTY()
+	int32 DerivedPartialValue = 999;
+
+	UPROPERTY(EditDefaultsOnly, Category = "VALUE")
+	TSubclassOf<UObject> ComponentClass;
+
+	UPROPERTY(EditDefaultsOnly, Category = "VALUE")
+	TSoftObjectPtr<UObject> SomeTexture;
+};
+
+UENUM()
+enum class ETestPartialEnum : uint8
+{
+	ValueA,
+	ValueB,
+	ValueC
+};
+
+UPARTIAL(UTestPartialObject)
+struct FAdvancedContainerPartial
+{
+	GENERATED_BODY()
+
+	UPROPERTY()
+	TMap<int32, FString> PartialMap;
+
+	UPROPERTY()
+	TSet<int32> PartialSet;
+
+	UPROPERTY()
+	ETestPartialEnum PartialEnum = ETestPartialEnum::ValueA;
+
+	FAdvancedContainerPartial()
+	{
+		PartialMap.Add(1, TEXT("One"));
+		PartialMap.Add(2, TEXT("Two"));
+
+		PartialSet.Add(10);
+		PartialSet.Add(20);
+		PartialSet.Add(30);
+	}
+};
+
+DECLARE_DYNAMIC_DELEGATE_OneParam(FTestPartialDelegate, int32, Value);
+DECLARE_DYNAMIC_MULTICAST_DELEGATE_OneParam(FTestPartialMulticastDelegate, int32, Value);
+
+UPARTIAL(UTestPartialObject)
+struct FAdvancedTypesPartial
+{
+	GENERATED_BODY()
+
+	UPROPERTY()
+	TObjectPtr<UClass> PartialClass = nullptr;
+
+	UPROPERTY()
+	uint16 PartialUInt16 = 12345;
+
+	UPROPERTY()
+	uint64 PartialUInt64 = 9876543210;
+
+	UPROPERTY()
+	FTestPartialDelegate PartialDelegate;
+
+	UPROPERTY()
+	FTestPartialMulticastDelegate PartialMulticastDelegate;
+};
+
+UPARTIAL(UTestPartialObject)
+struct FTestFunctionPartial
+{
+	GENERATED_BODY()
+
+	UPROPERTY()
+	int32 FunctionCallCount = 0;
+
+	UFUNCTION()
+	void IncrementCounter()
+	{
+		FunctionCallCount++;
+	}
+
+	UFUNCTION()
+	int32 GetCounter() const
+	{
+		return FunctionCallCount;
+	}
+
+	UFUNCTION()
+	int32 AddNumbers(int32 A, int32 B)
+	{
+		return A + B;
+	}
+
+	UFUNCTION()
+	void SetOwnerValue(int32 NewValue)
+	{
+		GetOwner().NativeValue = NewValue;
+	}
+};
+
+UPARTIAL(UTestPartialObject)
+struct FTestOutlinedFunctionPartial
+{
+	GENERATED_BODY()
+
+	UPROPERTY()
+	int32 StoredValue = 42;
+
+	UFUNCTION()
+	FOUNDATIONTESTS_API int32 GetStoredValue() const;
+
+	UFUNCTION()
+	FOUNDATIONTESTS_API void SetStoredValue(int32 NewValue);
+
+	UFUNCTION()
+	FOUNDATIONTESTS_API int32 MultiplyStoredValue(int32 Multiplier);
+
+	UFUNCTION()
+	FOUNDATIONTESTS_API int32 GetOwnerNativeValue() const;
+};
+
+UPARTIAL(UTestPartialObject)
+struct FTestEmptyPartial
+{
+	GENERATED_BODY()
+};
+
+UPARTIAL(UBlueprintablePartialObject)
+struct FBlueprintFunctionPartial
+{
+	GENERATED_BODY()
+
+	UPROPERTY(BlueprintReadWrite, Category = "PartialTest")
+	int32 Counter = 0;
+
+	UFUNCTION(BlueprintCallable, Category = "PartialTest")
+	void IncrementBPCounter()
+	{
+		Counter++;
+	}
+
+	UFUNCTION(BlueprintPure, Category = "PartialTest")
+	int32 GetBPCounter() const
+	{
+		return Counter;
+	}
+
+	UFUNCTION(BlueprintCallable, Category = "PartialTest")
+	int32 MultiplyBP(int32 A, int32 B)
+	{
+		return A * B;
+	}
+
+	UFUNCTION(BlueprintCallable, Category = "PartialTest")
+	void SetOwnerBlueprintValue(int32 NewValue)
+	{
+		GetOwner().BlueprintValue = NewValue;
+	}
+
+	UFUNCTION(BlueprintPure, Category = "PartialTest")
+	int32 GetOwnerBlueprintValue() const
+	{
+		return GetOwner().BlueprintValue;
+	}
+};
+
+
+UPARTIAL(USerializationTestObject)
+struct FSerializationTestPartial
+{
+	GENERATED_BODY()
+
+	UPROPERTY()
+	int32 PartialInt = 42;
+
+	UPROPERTY()
+	FString PartialString = TEXT("Partial");
+
+	UPROPERTY()
+	TArray<int32> PartialArray;
+
+	UPROPERTY()
+	TOptional<int32> PartialOptional;
+
+	UPROPERTY()
+	int32 PartialStaticArray[3];
+
+	FSerializationTestPartial()
+	{
+		PartialArray = {1, 2, 3};
+		PartialOptional = 999;
+		PartialStaticArray[0] = 10;
+		PartialStaticArray[1] = 20;
+		PartialStaticArray[2] = 30;
+	}
+};
+
+UPARTIAL(UDerivedSerializationTestObject)
+struct FDerivedSerializationTestPartial
+{
+	GENERATED_BODY()
+
+	UPROPERTY()
+	float DerivedPartialFloat = 3.14f;
+};
diff --git a/Engine/Source/Programs/LowLevelTests/FoundationTests/Tests/UPartialTests.cpp b/Engine/Source/Programs/LowLevelTests/FoundationTests/Tests/UPartialTests.cpp
new file mode 100644
--- /dev/null
+++ b/Engine/Source/Programs/LowLevelTests/FoundationTests/Tests/UPartialTests.cpp
@@ -0,0 +1,820 @@
+// Copyright Epic Games, Inc. All Rights Reserved.
+
+#include "TestHarness.h"
+#include "UObject/UObjectGlobals.h"
+#include "UObject/Package.h"
+#include "UObject/UObjectPartials.h"
+#include "UPartialTestClasses.h"
+#include "UPartialTestPartials.h"
+#include "Serialization/MemoryWriter.h"
+#include "Serialization/MemoryReader.h"
+#include "Serialization/ObjectAndNameAsStringProxyArchive.h"
+
+TEST_CASE("Partials BasicAccess", "[CoreUObject][Partials][UHT]")
+{
+	CollectGarbage(RF_NoFlags);
+
+	UTestPartialObject* TestObj = NewObject<UTestPartialObject>();
+	CHECK(TestObj != nullptr);
+
+	SECTION("Native property accessible")
+	{
+		CHECK(TestObj->NativeValue == 100);
+	}
+	SECTION("Partial accessible via GetPartial")
+	{
+		FTestObjectPartial& Partial = TestObj->GetPartial<FTestObjectPartial>();
+		CHECK(Partial.PartialValue == 43);
+		CHECK(Partial.PartialString == TEXT("PartialDefault"));
+		CHECK(Partial.PartialArray.Num() == 3);
+	}
+	SECTION("Second partial accessible")
+	{
+		FSecondObjectPartial& Partial = TestObj->GetPartial<FSecondObjectPartial>();
+		CHECK(Partial.PartialFloat == 3.14f);
+		CHECK(Partial.bPartialBool == true);
+	}
+	SECTION("Partial properties can be modified")
+	{
+		FTestObjectPartial& Partial = TestObj->GetPartial<FTestObjectPartial>();
+		Partial.PartialValue = 999;
+		Partial.PartialString = TEXT("Modified");
+	}
+
+	TestObj->AddToRoot();
+	CollectGarbage(RF_NoFlags);
+	TestObj->RemoveFromRoot();  
+	TWeakObjectPtr<UTestPartialObject> Weak = TestObj;
+	CollectGarbage(RF_NoFlags);
+	CHECK(!Weak.IsValid());
+}
+
+TEST_CASE("Partials Inheritance", "[CoreUObject][Partials][UHT]")
+{
+	UDerivedTestPartialObject* DerivedObj = NewObject<UDerivedTestPartialObject>();
+
+	SECTION("Native properties from base and derived accessible")
+	{
+		CHECK(DerivedObj->NativeValue == 100);
+		CHECK(DerivedObj->DerivedValue == 200);
+	}
+	SECTION("Base class partials accessible from derived instance")
+	{
+		FTestObjectPartial& Partial = DerivedObj->GetPartial<FTestObjectPartial>();
+		CHECK(Partial.PartialValue == 43);
+		CHECK(Partial.PartialString == TEXT("PartialDefault"));
+		CHECK(Partial.PartialArray.Num() == 3);
+		CHECK(Partial.PartialArray[0] == 1);
+		CHECK(Partial.PartialArray[1] == 2);
+		CHECK(Partial.PartialArray[2] == 3);
+
+		FSecondObjectPartial& SecondPartial = DerivedObj->GetPartial<FSecondObjectPartial>();
+		CHECK(SecondPartial.PartialFloat == 3.14f);
+	}
+	SECTION("Derived class partial accessible")
+	{
+		FDerivedObjectPartial& DerivedPartial = DerivedObj->GetPartial<FDerivedObjectPartial>();
+		CHECK(DerivedPartial.DerivedPartialValue == 999);
+	}
+}
+
+TEST_CASE("Partial Properties", "[CoreUObject][Partials][UHT]")
+{
+	UTestPartialObject* TestObj = NewObject<UTestPartialObject>();
+	UClass* TestClass = TestObj->GetClass();
+
+	SECTION("Partial properties are in class property list")
+	{
+		CHECK(TestClass->FindPropertyByName(FName("PartialValue")) != nullptr);
+		CHECK(TestClass->FindPropertyByName(FName("PartialString")) != nullptr);
+		CHECK(TestClass->FindPropertyByName(FName("PartialArray")) != nullptr);
+	}
+
+	SECTION("Partial properties have correct types")
+	{
+		CHECK(TestClass->FindPropertyByName(FName("PartialValue"))->IsA<FIntProperty>());
+		CHECK(TestClass->FindPropertyByName(FName("PartialString"))->IsA<FStrProperty>());
+		CHECK(TestClass->FindPropertyByName(FName("PartialArray"))->IsA<FArrayProperty>());
+	}
+
+	SECTION("Partial properties can be accessed via property system")
+	{
+		FProperty* PartialValueProp = TestClass->FindPropertyByName(FName("PartialValue"));
+		CHECK(PartialValueProp);
+		void* ValuePtr = PartialValueProp->ContainerPtrToValuePtr<void>(TestObj);
+		int32 Value = *(int32*)ValuePtr;
+		CHECK(Value == 43);
+		*(int32*)ValuePtr = 777;
+		CHECK(TestObj->GetPartial<FTestObjectPartial>().PartialValue == 777);
+	}
+}
+
+TEST_CASE("Partial AdvancedContainers", "[CoreUObject][Partials][UHT]")
+{
+	UTestPartialObject* TestObj = NewObject<UTestPartialObject>();
+
+	SECTION("Map partial property accessible")
+	{
+		FAdvancedContainerPartial& Partial = TestObj->GetPartial<FAdvancedContainerPartial>();
+		CHECK(Partial.PartialMap.Num() == 2);
+		CHECK(Partial.PartialMap.Contains(1));
+		CHECK(Partial.PartialMap.Contains(2));
+		CHECK(Partial.PartialMap[1] == TEXT("One"));
+		CHECK(Partial.PartialMap[2] == TEXT("Two"));
+	}
+
+	SECTION("Set partial property accessible")
+	{
+		FAdvancedContainerPartial& Partial = TestObj->GetPartial<FAdvancedContainerPartial>();
+		CHECK(Partial.PartialSet.Num() == 3);
+		CHECK(Partial.PartialSet.Contains(10));
+		CHECK(Partial.PartialSet.Contains(20));
+		CHECK(Partial.PartialSet.Contains(30));
+	}
+
+	SECTION("Enum partial property accessible")
+	{
+		FAdvancedContainerPartial& Partial = TestObj->GetPartial<FAdvancedContainerPartial>();
+		CHECK(Partial.PartialEnum == ETestPartialEnum::ValueA);
+		Partial.PartialEnum = ETestPartialEnum::ValueC;
+		CHECK(TestObj->GetPartial<FAdvancedContainerPartial>().PartialEnum == ETestPartialEnum::ValueC);
+	}
+
+	SECTION("Map property is in class property list")
+	{
+		FProperty* MapProp = TestObj->GetClass()->FindPropertyByName(FName("PartialMap"));
+		CHECK(MapProp != nullptr);
+		CHECK(MapProp->IsA<FMapProperty>());
+	}
+
+	SECTION("Set property is in class property list")
+	{
+		FProperty* SetProp = TestObj->GetClass()->FindPropertyByName(FName("PartialSet"));
+		CHECK(SetProp != nullptr);
+		CHECK(SetProp->IsA<FSetProperty>());
+	}
+
+	SECTION("Enum property is in class property list")
+	{
+		UClass* TestClass = TestObj->GetClass();
+		FProperty* EnumProp = TestClass->FindPropertyByName(FName("PartialEnum"));
+		CHECK(EnumProp != nullptr);
+		CHECK(EnumProp->IsA<FEnumProperty>());
+	}
+}
+
+TEST_CASE("Partial PropertySerialization", "[CoreUObject][Partials][UHT]")
+{
+	UTestPartialObject* TestObj = NewObject<UTestPartialObject>();
+
+	FTestObjectPartial& Partial = TestObj->GetPartial<FTestObjectPartial>();
+	Partial.PartialValue = 12345;
+	Partial.PartialString = TEXT("SerializedValue");
+	Partial.PartialArray.Empty();
+	Partial.PartialArray.Add(99);
+	Partial.PartialArray.Add(88);
+	Partial.PartialArray.Add(77);
+
+	SECTION("Partial properties can be serialized via property system")
+	{
+		UClass* TestClass = TestObj->GetClass();
+
+		// Serialize int property
+		FProperty* IntProp = TestClass->FindPropertyByName(FName("PartialValue"));
+		CHECK(IntProp != nullptr);
+
+		TArray<uint8> IntData;
+		FMemoryWriter IntWriter(IntData);
+		void* IntValuePtr = IntProp->ContainerPtrToValuePtr<void>(TestObj);
+		IntProp->SerializeItem(FStructuredArchiveFromArchive(IntWriter).GetSlot(), IntValuePtr);
+		CHECK(IntData.Num() > 0);
+
+		// Deserialize to verify
+		FMemoryReader IntReader(IntData);
+		int32 DeserializedInt = 0;
+		IntProp->SerializeItem(FStructuredArchiveFromArchive(IntReader).GetSlot(), &DeserializedInt);
+		CHECK(DeserializedInt == 12345);
+	}
+
+	SECTION("Partial string properties can be serialized")
+	{
+		UClass* TestClass = TestObj->GetClass();
+		FProperty* StringProp = TestClass->FindPropertyByName(FName("PartialString"));
+		CHECK(StringProp != nullptr);
+
+		TArray<uint8> StringData;
+		FMemoryWriter StringWriter(StringData);
+		void* StringValuePtr = StringProp->ContainerPtrToValuePtr<void>(TestObj);
+		StringProp->SerializeItem(FStructuredArchiveFromArchive(StringWriter).GetSlot(), StringValuePtr);
+		CHECK(StringData.Num() > 0);
+	}
+
+	SECTION("Partial array properties can be serialized")
+	{
+		UClass* TestClass = TestObj->GetClass();
+		FProperty* ArrayProp = TestClass->FindPropertyByName(FName("PartialArray"));
+		CHECK(ArrayProp != nullptr);
+
+		TArray<uint8> ArrayData;
+		FMemoryWriter ArrayWriter(ArrayData);
+		void* ArrayValuePtr = ArrayProp->ContainerPtrToValuePtr<void>(TestObj);
+		ArrayProp->SerializeItem(FStructuredArchiveFromArchive(ArrayWriter).GetSlot(), ArrayValuePtr);
+		CHECK(ArrayData.Num() > 0);
+	}
+}
+
+TEST_CASE("Partial AdvancedTypes", "[CoreUObject][Partials][UHT]")
+{
+	UTestPartialObject* TestObj = NewObject<UTestPartialObject>();
+	CHECK(TestObj != nullptr);
+
+	SECTION("UClass* Partial property accessible")
+	{
+		FAdvancedTypesPartial& Partial = TestObj->GetPartial<FAdvancedTypesPartial>();
+		CHECK(Partial.PartialClass == nullptr);
+		Partial.PartialClass = UObject::StaticClass();
+		CHECK(Partial.PartialClass == UObject::StaticClass());
+		CHECK(TestObj->GetPartial<FAdvancedTypesPartial>().PartialClass == UObject::StaticClass());
+	}
+
+	SECTION("uint16 Partial property accessible")
+	{
+		FAdvancedTypesPartial& Partial = TestObj->GetPartial<FAdvancedTypesPartial>();
+		CHECK(Partial.PartialUInt16 == 12345);
+		Partial.PartialUInt16 = 54321;
+	}
+
+	SECTION("uint64 Partial property accessible")
+	{
+		FAdvancedTypesPartial& Partial = TestObj->GetPartial<FAdvancedTypesPartial>();
+		CHECK(Partial.PartialUInt64 == 9876543210);
+		Partial.PartialUInt64 = 1234567890123;
+	}
+
+	SECTION("Delegate Partial property accessible")
+	{
+		FAdvancedTypesPartial& Partial = TestObj->GetPartial<FAdvancedTypesPartial>();
+		CHECK(!Partial.PartialDelegate.IsBound());
+	}
+
+	SECTION("Class property is in class property list")
+	{
+		UClass* TestClass = TestObj->GetClass();
+		FProperty* ClassProp = TestClass->FindPropertyByName(FName("PartialClass"));
+		CHECK(ClassProp != nullptr);
+		CHECK(ClassProp->IsA<FClassProperty>());
+	}
+
+	SECTION("uint16 property is in class property list")
+	{
+		UClass* TestClass = TestObj->GetClass();
+		FProperty* UInt16Prop = TestClass->FindPropertyByName(FName("PartialUInt16"));
+		CHECK(UInt16Prop != nullptr);
+		CHECK(UInt16Prop->IsA<FUInt16Property>());
+	}
+
+	SECTION("uint64 property is in class property list")
+	{
+		UClass* TestClass = TestObj->GetClass();
+		FProperty* UInt64Prop = TestClass->FindPropertyByName(FName("PartialUInt64"));
+		CHECK(UInt64Prop != nullptr);
+		CHECK(UInt64Prop->IsA<FUInt64Property>());
+	}
+
+	SECTION("Delegate property is in class property list")
+	{
+		UClass* TestClass = TestObj->GetClass();
+		FProperty* DelegateProp = TestClass->FindPropertyByName(FName("PartialDelegate"));
+		CHECK(DelegateProp != nullptr);
+		CHECK(DelegateProp->IsA<FDelegateProperty>());
+	}
+
+	SECTION("Multicast delegate property is in class property list")
+	{
+		UClass* TestClass = TestObj->GetClass();
+		FProperty* MulticastProp = TestClass->FindPropertyByName(FName("PartialMulticastDelegate"));
+		CHECK(MulticastProp != nullptr);
+		CHECK(MulticastProp->IsA<FMulticastDelegateProperty>());
+	}
+}
+
+TEST_CASE("Partial GetOwner", "[CoreUObject][Partials][GetOwner]")
+{
+	UTestPartialObject* TestObj = NewObject<UTestPartialObject>();
+	TestObj->NativeValue = 12345;
+
+	SECTION("GetOwner() returns correct object reference")
+	{
+		FTestObjectPartial& Partial = TestObj->GetPartial<FTestObjectPartial>();
+		UTestPartialObject& Owner = Partial.GetOwner();
+		CHECK(&Owner == TestObj);
+		CHECK(Owner.NativeValue == 12345);
+	}
+
+	SECTION("GetOwner() const returns correct object reference")
+	{
+		const FTestObjectPartial& Partial = TestObj->GetPartial<FTestObjectPartial>();
+		const UTestPartialObject& Owner = Partial.GetOwner();
+		CHECK(&Owner == TestObj);
+		CHECK(Owner.NativeValue == 12345);
+	}
+
+	SECTION("GetOwner() allows modifying owner properties")
+	{
+		FTestObjectPartial& Partial = TestObj->GetPartial<FTestObjectPartial>();
+		Partial.GetOwner().NativeValue = 99999;
+		CHECK(TestObj->NativeValue == 99999);
+	}
+
+	SECTION("Multiple Partials can access same owner")
+	{
+		FTestObjectPartial& Partial1 = TestObj->GetPartial<FTestObjectPartial>();
+		FSecondObjectPartial& Partial2 = TestObj->GetPartial<FSecondObjectPartial>();
+		CHECK(&Partial1.GetOwner() == TestObj);
+		CHECK(&Partial2.GetOwner() == TestObj);
+		CHECK(&Partial1.GetOwner() == &Partial2.GetOwner());
+	}
+
+	SECTION("GetOwner() works with derived classes")
+	{
+		UDerivedTestPartialObject* DerivedObj = NewObject<UDerivedTestPartialObject>();
+		DerivedObj->NativeValue = 111;
+		DerivedObj->DerivedValue = 222;
+		FDerivedObjectPartial& DerivedPartial = DerivedObj->GetPartial<FDerivedObjectPartial>();
+		UDerivedTestPartialObject& Owner = DerivedPartial.GetOwner();
+		CHECK(&Owner == DerivedObj);
+		CHECK(Owner.NativeValue == 111);
+		CHECK(Owner.DerivedValue == 222);
+	}
+}
+
+TEST_CASE("Partials UFunctions", "[CoreUObject][Partials][UFunctions]")
+{
+	UTestPartialObject* TestObj = NewObject<UTestPartialObject>();
+	TestObj->NativeValue = 100;
+
+	SECTION("Functions are discoverable via FindFunction")
+	{
+		UClass* Class = TestObj->GetClass();
+		CHECK(Class->FindFunctionByName(FName("IncrementCounter")));
+		CHECK(Class->FindFunctionByName(FName("GetCounter")));
+		CHECK(Class->FindFunctionByName(FName("AddNumbers")));
+		CHECK(Class->FindFunctionByName(FName("SetOwnerValue")));
+	}
+
+	SECTION("Function with no parameters can be called via ProcessEvent")
+	{
+		FTestFunctionPartial& Partial = TestObj->GetPartial<FTestFunctionPartial>();
+		CHECK(Partial.FunctionCallCount == 0);
+		UFunction* IncrementFunc = TestObj->GetClass()->FindFunctionByName(FName("IncrementCounter"));
+		CHECK(IncrementFunc != nullptr);
+		TestObj->ProcessEvent(IncrementFunc, nullptr);
+		CHECK(Partial.FunctionCallCount == 1);
+		TestObj->ProcessEvent(IncrementFunc, nullptr);
+		CHECK(Partial.FunctionCallCount == 2);
+	}
+
+	SECTION("Function with return value can be called via ProcessEvent")
+	{
+		FTestFunctionPartial& Partial = TestObj->GetPartial<FTestFunctionPartial>();
+		CHECK(Partial.FunctionCallCount == 0);
+		Partial.FunctionCallCount = 42;
+		UFunction* GetCounterFunc = TestObj->GetClass()->FindFunctionByName(FName("GetCounter"));
+		CHECK(GetCounterFunc != nullptr);
+		struct { int32 ReturnValue; } Params;
+		TestObj->ProcessEvent(GetCounterFunc, &Params);
+		CHECK(Params.ReturnValue == 42);
+	}
+
+	SECTION("Function with parameters can be called via ProcessEvent")
+	{
+		UFunction* AddNumbersFunc = TestObj->GetClass()->FindFunctionByName(FName("AddNumbers"));
+		CHECK(AddNumbersFunc != nullptr);
+		struct { int32 A; int32 B; int32 ReturnValue; } Params;
+		Params.A = 15;
+		Params.B = 27;
+		Params.ReturnValue = 0;
+		TestObj->ProcessEvent(AddNumbersFunc, &Params);
+		CHECK(Params.ReturnValue == 42);
+	}
+
+	SECTION("Function can modify owner via GetOwner()")
+	{
+		CHECK(TestObj->NativeValue == 100);
+		UFunction* SetOwnerValueFunc = TestObj->GetClass()->FindFunctionByName(FName("SetOwnerValue"));
+		CHECK(SetOwnerValueFunc != nullptr);
+		struct { int32 NewValue; } Params;
+		Params.NewValue = 999;
+		TestObj->ProcessEvent(SetOwnerValueFunc, &Params);
+		CHECK(TestObj->NativeValue == 999);
+	}
+}
+
+TEST_CASE("Partials OutlinedUFunctions", "[CoreUObject][Partials][UFunctions][Outlined]")
+{
+	UTestPartialObject* TestObj = NewObject<UTestPartialObject>();
+	TestObj->NativeValue = 200;
+
+	SECTION("Outlined Partial functions are discoverable via FindFunction")
+	{
+		UClass* Class = TestObj->GetClass();
+		CHECK(Class->FindFunctionByName(FName("GetStoredValue")));
+		CHECK(Class->FindFunctionByName(FName("SetStoredValue")));
+		CHECK(Class->FindFunctionByName(FName("MultiplyStoredValue")));
+		CHECK(Class->FindFunctionByName(FName("GetOwnerNativeValue")));
+	}
+
+	SECTION("Outlined Partial function with return value can be called via ProcessEvent")
+	{
+		FTestOutlinedFunctionPartial& Partial = TestObj->GetPartial<FTestOutlinedFunctionPartial>();
+		CHECK(Partial.StoredValue == 42);
+		UFunction* GetStoredValueFunc = TestObj->GetClass()->FindFunctionByName(FName("GetStoredValue"));
+		CHECK(GetStoredValueFunc != nullptr);
+		struct { int32 ReturnValue; } Params;
+		TestObj->ProcessEvent(GetStoredValueFunc, &Params);
+		CHECK(Params.ReturnValue == 42);
+	}
+
+	SECTION("Outlined Partial function with parameter can modify Partial state")
+	{
+		FTestOutlinedFunctionPartial& Partial = TestObj->GetPartial<FTestOutlinedFunctionPartial>();
+		CHECK(Partial.StoredValue == 42);
+		UFunction* SetStoredValueFunc = TestObj->GetClass()->FindFunctionByName(FName("SetStoredValue"));
+		CHECK(SetStoredValueFunc != nullptr);
+		struct { int32 NewValue; } Params;
+		Params.NewValue = 123;
+		TestObj->ProcessEvent(SetStoredValueFunc, &Params);
+		CHECK(Partial.StoredValue == 123);
+	}
+
+	SECTION("Outlined Partial function with parameter and return value works via ProcessEvent")
+	{
+		FTestOutlinedFunctionPartial& Partial = TestObj->GetPartial<FTestOutlinedFunctionPartial>();
+		Partial.StoredValue = 10;
+		UFunction* MultiplyFunc = TestObj->GetClass()->FindFunctionByName(FName("MultiplyStoredValue"));
+		CHECK(MultiplyFunc != nullptr);
+		struct { int32 Multiplier; int32 ReturnValue; } Params;
+		Params.Multiplier = 5;
+		Params.ReturnValue = 0;
+		TestObj->ProcessEvent(MultiplyFunc, &Params);
+		CHECK(Params.ReturnValue == 50);
+		CHECK(Partial.StoredValue == 50);
+	}
+
+	SECTION("Outlined Partial function can access owner via GetOwner()")
+	{
+		TestObj->NativeValue = 200;
+		UFunction* GetOwnerNativeValueFunc = TestObj->GetClass()->FindFunctionByName(FName("GetOwnerNativeValue"));
+		CHECK(GetOwnerNativeValueFunc != nullptr);
+		struct { int32 ReturnValue; } Params;
+		TestObj->ProcessEvent(GetOwnerNativeValueFunc, &Params);
+		CHECK(Params.ReturnValue == 200);
+	}
+}
+
+TEST_CASE("Partials BlueprintCallable", "[CoreUObject][Partials][UHT][Blueprint]")
+{
+	UBlueprintablePartialObject* TestObj = NewObject<UBlueprintablePartialObject>(GetTransientPackage(), "TestBPObject");
+
+	SECTION("BlueprintCallable Partial functions are discoverable via FindFunction")
+	{
+		UClass* Class = UBlueprintablePartialObject::StaticClass();
+		CHECK(Class->FindFunctionByName(FName("IncrementBPCounter")));
+		CHECK(Class->FindFunctionByName(FName("GetBPCounter")));
+		CHECK(Class->FindFunctionByName(FName("MultiplyBP")));
+		CHECK(Class->FindFunctionByName(FName("SetOwnerBlueprintValue")));
+		CHECK(Class->FindFunctionByName(FName("GetOwnerBlueprintValue")));
+	}
+
+	SECTION("BlueprintCallable Partial function has correct flags")
+	{
+		UFunction* IncrementFunc = TestObj->GetClass()->FindFunctionByName(FName("IncrementBPCounter"));
+		CHECK(IncrementFunc->HasAnyFunctionFlags(FUNC_BlueprintCallable));
+	}
+
+	SECTION("BlueprintPure Partial function has correct flags")
+	{
+		UFunction* GetCounterFunc = TestObj->GetClass()->FindFunctionByName(FName("GetBPCounter"));
+		CHECK(GetCounterFunc->HasAnyFunctionFlags(FUNC_BlueprintCallable));
+		CHECK(GetCounterFunc->HasAnyFunctionFlags(FUNC_BlueprintPure));
+	}
+
+	SECTION("BlueprintCallable Partial function works via ProcessEvent")
+	{
+		FBlueprintFunctionPartial& Partial = TestObj->GetPartial<FBlueprintFunctionPartial>();
+		CHECK(Partial.Counter == 0);
+		UFunction* IncrementFunc = TestObj->GetClass()->FindFunctionByName(FName("IncrementBPCounter"));
+		TestObj->ProcessEvent(IncrementFunc, nullptr);
+		CHECK(Partial.Counter == 1);
+		TestObj->ProcessEvent(IncrementFunc, nullptr);
+		CHECK(Partial.Counter == 2);
+	}
+
+	SECTION("BlueprintPure Partial function with return value works via ProcessEvent")
+	{
+		FBlueprintFunctionPartial& Partial = TestObj->GetPartial<FBlueprintFunctionPartial>();
+		Partial.Counter = 42;
+		UFunction* GetCounterFunc = TestObj->GetClass()->FindFunctionByName(FName("GetBPCounter"));
+		struct { int32 ReturnValue; } Params;
+		Params.ReturnValue = 0;
+		TestObj->ProcessEvent(GetCounterFunc, &Params);
+		CHECK(Params.ReturnValue == 42);
+	}
+
+	SECTION("BlueprintCallable Partial function with parameters and return value works")
+	{
+		UFunction* MultiplyFunc = TestObj->GetClass()->FindFunctionByName(FName("MultiplyBP"));
+		struct { int32 A; int32 B; int32 ReturnValue; } Params;
+		Params.A = 7;
+		Params.B = 8;
+		Params.ReturnValue = 0;
+		TestObj->ProcessEvent(MultiplyFunc, &Params);
+		CHECK(Params.ReturnValue == 56);
+	}
+
+	SECTION("BlueprintCallable Partial function can access and modify owner")
+	{
+		TestObj->BlueprintValue = 100;
+		UFunction* SetOwnerFunc = TestObj->GetClass()->FindFunctionByName(FName("SetOwnerBlueprintValue"));
+		struct { int32 NewValue; } Params;
+		Params.NewValue = 999;
+		TestObj->ProcessEvent(SetOwnerFunc, &Params);
+		CHECK(TestObj->BlueprintValue == 999);
+	}
+
+	SECTION("BlueprintPure Partial function can read owner value")
+	{
+		TestObj->BlueprintValue = 777;
+		UFunction* GetOwnerFunc = TestObj->GetClass()->FindFunctionByName(FName("GetOwnerBlueprintValue"));
+		struct { int32 ReturnValue; } Params;
+		Params.ReturnValue = 0;
+		TestObj->ProcessEvent(GetOwnerFunc, &Params);
+		CHECK(Params.ReturnValue == 777);
+	}
+}
+
+TEST_CASE("Partials Lifecycle", "[CoreUObject][Partials][UHT]")
+{
+	SECTION("Partial constructors are called for CDO")
+	{
+		UClass* TestClass = UTestPartialObject::StaticClass();
+		UObject* CDO = TestClass->GetDefaultObject();
+		CHECK(CDO != nullptr);
+		UTestPartialObject* TestCDO = Cast<UTestPartialObject>(CDO);
+		CHECK(TestCDO != nullptr);
+		FTestObjectPartial& Partial = TestCDO->GetPartial<FTestObjectPartial>();
+		CHECK(Partial.PartialValue == 43);
+		CHECK(Partial.PartialString == TEXT("PartialDefault"));
+		CHECK(Partial.PartialArray.Num() == 3);
+		CHECK(Partial.PartialArray[0] == 1);
+		CHECK(Partial.PartialArray[1] == 2);
+		CHECK(Partial.PartialArray[2] == 3);
+	}
+
+	SECTION("Partial destructors are called when object is destroyed")
+	{
+		FName PackageName = MakeUniqueObjectName(nullptr, UPackage::StaticClass(), FName("/Memory/DestructorTest"));
+		UPackage* Package = NewObject<UPackage>(nullptr, UPackage::StaticClass(), PackageName);
+		{
+			UTestPartialObject* TestObj = NewObject<UTestPartialObject>(Package, "TestObject");
+			CHECK(TestObj != nullptr);
+			FAdvancedContainerPartial& Partial = TestObj->GetPartial<FAdvancedContainerPartial>();
+			Partial.PartialMap.Add(1, TEXT("Test"));
+			Partial.PartialSet.Add(42);
+		}
+		Package->MarkAsGarbage();
+	}
+}
+
+// ============================================================================
+// Property Merging Tests - verify properties are added to UClass correctly
+// ============================================================================
+
+TEST_CASE("Partials PropertyMerging", "[CoreUObject][Partials][UHT]")
+{
+	UClass* TestClass = UTestPartialObject::StaticClass();
+
+	SECTION("Partial properties added to UClass property list")
+	{
+		int32 PropertyCount = 0;
+		for (FField* Field = TestClass->ChildProperties; Field; Field = Field->Next)
+		{
+			if (CastField<FProperty>(Field))
+			{
+				PropertyCount++;
+			}
+		}
+		CHECK(PropertyCount > 4);
+
+		bool bFoundPartialValue = false;
+		bool bFoundPartialString = false;
+		bool bFoundPartialArray = false;
+		for (FField* Field = TestClass->ChildProperties; Field; Field = Field->Next)
+		{
+			if (FProperty* Prop = CastField<FProperty>(Field))
+			{
+				FString PropName = Prop->GetName();
+				if (PropName == TEXT("PartialValue")) bFoundPartialValue = true;
+				if (PropName == TEXT("PartialString")) bFoundPartialString = true;
+				if (PropName == TEXT("PartialArray")) bFoundPartialArray = true;
+			}
+		}
+		CHECK(bFoundPartialValue);
+		CHECK(bFoundPartialString);
+		CHECK(bFoundPartialArray);
+	}
+
+	SECTION("Partial properties can be found by name")
+	{
+		FProperty* PartialValueProp = TestClass->FindPropertyByName(FName("PartialValue"));
+		CHECK(PartialValueProp != nullptr);
+		CHECK(PartialValueProp->IsA<FIntProperty>());
+		FProperty* PartialStringProp = TestClass->FindPropertyByName(FName("PartialString"));
+		CHECK(PartialStringProp != nullptr);
+		CHECK(PartialStringProp->IsA<FStrProperty>());
+		FProperty* PartialArrayProp = TestClass->FindPropertyByName(FName("PartialArray"));
+		CHECK(PartialArrayProp != nullptr);
+		CHECK(PartialArrayProp->IsA<FArrayProperty>());
+	}
+
+	SECTION("Partial properties have correct offsets for reflection")
+	{
+		UTestPartialObject* TestObj = NewObject<UTestPartialObject>();
+		FProperty* PartialValueProp = TestClass->FindPropertyByName(FName("PartialValue"));
+		void* PropValuePtr = PartialValueProp->ContainerPtrToValuePtr<void>(TestObj);
+		int32* ValuePtr = static_cast<int32*>(PropValuePtr);
+		FTestObjectPartial& Partial = TestObj->GetPartial<FTestObjectPartial>();
+		CHECK(*ValuePtr == Partial.PartialValue);
+		CHECK(*ValuePtr == 43); // Default value from constructor
+	}
+
+	SECTION("Derived class includes base class Partial properties")
+	{
+		UClass* DerivedClass = UDerivedTestPartialObject::StaticClass();
+		UClass* BaseClass = UTestPartialObject::StaticClass();
+		FProperty* BasePartialValue = DerivedClass->FindPropertyByName(FName("PartialValue"));
+		FProperty* DerivedPartialValue = DerivedClass->FindPropertyByName(FName("DerivedPartialValue"));
+		CHECK(DerivedClass->GetPropertiesSize() > BaseClass->GetPropertiesSize());
+	}
+}
+
+TEST_CASE("Partials ObjectDuplication", "[CoreUObject][Partials][UHT]")
+{
+	UTestPartialObject* SourceObj = NewObject<UTestPartialObject>(GetTransientPackage(), "SourceObject");
+	SourceObj->NativeValue = 12345;
+	FTestObjectPartial& SourcePartial = SourceObj->GetPartial<FTestObjectPartial>();
+	SourcePartial.PartialValue = 99999;
+	SourcePartial.PartialString = TEXT("DuplicationTest");
+	UTestPartialObject* DuplicatedObj = NewObject<UTestPartialObject>(GetTransientPackage(), "DuplicatedObject");
+	CHECK(DuplicatedObj != nullptr);
+	CHECK(DuplicatedObj != SourceObj);
+	DuplicatedObj->NativeValue = SourceObj->NativeValue;
+	FTestObjectPartial& DupPartial = DuplicatedObj->GetPartial<FTestObjectPartial>();
+	DupPartial.PartialValue = SourcePartial.PartialValue;
+	DupPartial.PartialString = SourcePartial.PartialString;
+}
+
+TEST_CASE("Partials InheritanceWithoutPartialAttribute", "[CoreUObject][Partials][UHT]")
+{
+	UDerivedDerivedTestPartialObject* DerivedDerivedObj = NewObject<UDerivedDerivedTestPartialObject>(GetTransientPackage(), "DerivedDerivedObject");
+
+	SECTION("Native properties from all levels accessible")
+	{
+		CHECK(DerivedDerivedObj->NativeValue == 100);
+		CHECK(DerivedDerivedObj->DerivedValue == 200);
+	}
+
+	SECTION("Grandparent class Partials should be accessible (CURRENTLY FAILS)")
+	{
+		FTestObjectPartial& BasePartial = DerivedDerivedObj->GetPartial<FTestObjectPartial>();
+		CHECK(BasePartial.PartialValue == 43);
+		CHECK(BasePartial.PartialString == TEXT("PartialDefault"));
+		CHECK(BasePartial.PartialArray.Num() == 3);
+	}
+
+	SECTION("Parent class Partials should be accessible (CURRENTLY FAILS)")
+	{
+		FDerivedObjectPartial& DerivedPartial = DerivedDerivedObj->GetPartial<FDerivedObjectPartial>();
+		CHECK(DerivedPartial.DerivedPartialValue == 999);
+	}
+
+	SECTION("Partial properties should appear in reflection")
+	{
+		UClass* Class = DerivedDerivedObj->GetClass();
+		FProperty* PartialValueProp = Class->FindPropertyByName(FName("PartialValue"));
+		CHECK(PartialValueProp != nullptr);
+		FProperty* DerivedPartialValueProp = Class->FindPropertyByName(FName("DerivedPartialValue"));
+		CHECK(DerivedPartialValueProp != nullptr);
+	}
+}
+
+TEST_CASE("Partials PartialProperties", "[CoreUObject][Serialization][Partials]")
+{
+	USerializationTestObject* OriginalObj = NewObject<USerializationTestObject>();
+	OriginalObj->NativeProperty = 123;
+	OriginalObj->NativeString = TEXT("Modified");
+
+	FSerializationTestPartial& Partial = OriginalObj->GetPartial<FSerializationTestPartial>();
+	Partial.PartialInt = 456;
+	Partial.PartialString = TEXT("ModifiedPartial");
+	Partial.PartialArray = {10, 20, 30, 40};
+	Partial.PartialOptional = 777;
+
+	SECTION("Unversioned serialization preserves partial properties")
+	{
+		TArray<uint8> SaveData;
+		FMemoryWriter Writer(SaveData);
+		Writer.SetUseUnversionedPropertySerialization(true);
+		OriginalObj->Serialize(Writer);
+		USerializationTestObject* LoadedObj = NewObject<USerializationTestObject>();
+		FMemoryReader Reader(SaveData);
+		Reader.SetUseUnversionedPropertySerialization(true);
+		LoadedObj->Serialize(Reader);
+		CHECK(LoadedObj->NativeProperty == 123);
+		CHECK(LoadedObj->NativeString == TEXT("Modified"));
+		FSerializationTestPartial& LoadedPartial = LoadedObj->GetPartial<FSerializationTestPartial>();
+		CHECK(LoadedPartial.PartialInt == 456);
+		CHECK(LoadedPartial.PartialString == TEXT("ModifiedPartial"));
+		CHECK(LoadedPartial.PartialArray.Num() == 4);
+		CHECK(LoadedPartial.PartialArray[0] == 10);
+		CHECK(LoadedPartial.PartialArray[3] == 40);
+		CHECK(LoadedPartial.PartialOptional.IsSet());
+		CHECK(LoadedPartial.PartialOptional.GetValue() == 777);
+	}
+
+	SECTION("Default values work correctly for partial properties")
+	{
+		USerializationTestObject* DefaultObj = NewObject<USerializationTestObject>();
+		TArray<uint8> SaveData;
+		FMemoryWriter Writer(SaveData);
+		Writer.SetUseUnversionedPropertySerialization(true);
+		DefaultObj->Serialize(Writer);
+		USerializationTestObject* LoadedObj = NewObject<USerializationTestObject>();
+		FMemoryReader Reader(SaveData);
+		Reader.SetUseUnversionedPropertySerialization(true);
+		LoadedObj->Serialize(Reader);
+		FSerializationTestPartial& LoadedPartial = LoadedObj->GetPartial<FSerializationTestPartial>();
+		CHECK(LoadedPartial.PartialInt == 42);
+		CHECK(LoadedPartial.PartialString == TEXT("Partial"));
+	}
+
+	SECTION("Static arrays in partials serialize correctly")
+	{
+		USerializationTestObject* Obj = NewObject<USerializationTestObject>();
+		FSerializationTestPartial& Partial2 = Obj->GetPartial<FSerializationTestPartial>();
+		Partial2.PartialStaticArray[0] = 100;
+		Partial2.PartialStaticArray[1] = 200;
+		Partial2.PartialStaticArray[2] = 300;
+		TArray<uint8> SaveData;
+		FMemoryWriter Writer(SaveData);
+		Writer.SetUseUnversionedPropertySerialization(true);
+		Obj->Serialize(Writer);
+		USerializationTestObject* LoadedObj = NewObject<USerializationTestObject>();
+		FMemoryReader Reader(SaveData);
+		Reader.SetUseUnversionedPropertySerialization(true);
+		LoadedObj->Serialize(Reader);
+		FSerializationTestPartial& LoadedPartial = LoadedObj->GetPartial<FSerializationTestPartial>();
+		CHECK(LoadedPartial.PartialStaticArray[0] == 100);
+		CHECK(LoadedPartial.PartialStaticArray[1] == 200);
+		CHECK(LoadedPartial.PartialStaticArray[2] == 300);
+	}
+}
+
+TEST_CASE("Partials InheritedPartialProperties", "[CoreUObject][Serialization][Partials]")
+{
+	SECTION("Derived class serializes both base and derived partial properties")
+	{
+		UDerivedSerializationTestObject* OriginalObj = NewObject<UDerivedSerializationTestObject>();
+		FSerializationTestPartial& BasePartial = OriginalObj->GetPartial<FSerializationTestPartial>();
+		BasePartial.PartialInt = 111;
+		FDerivedSerializationTestPartial& DerivedPartial = OriginalObj->GetPartial<FDerivedSerializationTestPartial>();
+		DerivedPartial.DerivedPartialFloat = 2.71f;
+		TArray<uint8> SaveData;
+		FMemoryWriter Writer(SaveData);
+		Writer.SetUseUnversionedPropertySerialization(true);
+		OriginalObj->Serialize(Writer);
+		UDerivedSerializationTestObject* LoadedObj = NewObject<UDerivedSerializationTestObject>();
+		FMemoryReader Reader(SaveData);
+		Reader.SetUseUnversionedPropertySerialization(true);
+		LoadedObj->Serialize(Reader);
+		FSerializationTestPartial& LoadedBasePartial = LoadedObj->GetPartial<FSerializationTestPartial>();
+		CHECK(LoadedBasePartial.PartialInt == 111);
+		FDerivedSerializationTestPartial& LoadedDerivedPartial = LoadedObj->GetPartial<FDerivedSerializationTestPartial>();
+		CHECK(LoadedDerivedPartial.DerivedPartialFloat == 2.71f);
+	}
+}
+
+TEST_CASE("Partials GetDefaultValue_NegativeOffset", "[CoreUObject][Serialization][Partials]")
+{
+	USerializationTestObject* Obj = NewObject<USerializationTestObject>();
+	FSerializationTestPartial& Partial = Obj->GetPartial<FSerializationTestPartial>();
+	Partial.PartialString = TEXT("ModifiedValue");
+	TArray<uint8> SaveData;
+	FMemoryWriter Writer(SaveData);
+	Writer.SetUseUnversionedPropertySerialization(true);
+	FStructuredArchiveFromArchive Adapter(Writer);
+	UClass* Class = Obj->GetClass();
+	const void* Defaults = Class->GetDefaultObject();
+	Class->SerializeTaggedProperties(Adapter.GetSlot(), (uint8*)Obj, Class, (uint8*)Defaults);
+	CHECK(SaveData.Num() > 0);
+}
\ No newline at end of file
diff --git a/Engine/Source/Programs/Shared/EpicGames.Core/UnrealEngineTypes.cs b/Engine/Source/Programs/Shared/EpicGames.Core/UnrealEngineTypes.cs
--- a/Engine/Source/Programs/Shared/EpicGames.Core/UnrealEngineTypes.cs
+++ b/Engine/Source/Programs/Shared/EpicGames.Core/UnrealEngineTypes.cs
@@ -158,6 +158,11 @@ namespace EpicGames.Core
 		/// </summary>
 		Native = 0x00000080u,
 
+		/// <summary>
+		/// Class is a partial class
+		/// </summary>
+		Partial = 0x00000100u,
+
 		/// <summary>
 		/// Do not allow users to create in the editor.
 		/// </summary>
@@ -276,7 +281,7 @@ namespace EpicGames.Core
 		/// <summary>
 		/// Flags to inherit from base class
 		/// </summary>
-		Inherit = Transient | Optional | DefaultConfig | Config | PerObjectConfig | ConfigDoNotCheckDefaults | NotPlaceable | Const | HasInstancedReference |
+		Inherit = Transient | Optional | DefaultConfig | Config | Partial | PerObjectConfig | ConfigDoNotCheckDefaults | NotPlaceable | Const | HasInstancedReference |
 			Deprecated | DefaultToInstanced | GlobalUserConfig | ProjectUserConfig | NeedsDeferredDependencyLoading,
 
 		/// <summary>
@@ -297,7 +302,7 @@ namespace EpicGames.Core
 		/// <summary>
 		/// This is used as a mask for the flags put into generated code for "compiled in" classes.
 		/// </summary>
-		SaveInCompiledInClasses = Abstract | DefaultConfig | GlobalUserConfig | ProjectUserConfig | PerPlatformConfig | Config | Transient | Optional | Native | NotPlaceable | PerObjectConfig |
+		SaveInCompiledInClasses = Abstract | DefaultConfig | GlobalUserConfig | ProjectUserConfig | PerPlatformConfig | Config | Transient | Optional | Native | Partial | NotPlaceable | PerObjectConfig |
 			ConfigDoNotCheckDefaults | EditInlineNew | CollapseCategories | Interface | DefaultToInstanced | HasInstancedReference | Hidden | Deprecated |
 			HideDropDown | Intrinsic | Const | MinimalAPI | RequiredAPI | MatchedSerializers | NeedsDeferredDependencyLoading,
 	};
diff --git a/Engine/Source/Programs/Shared/EpicGames.UHT/Exporters/CodeGen/UhtCodeGenerator.cs b/Engine/Source/Programs/Shared/EpicGames.UHT/Exporters/CodeGen/UhtCodeGenerator.cs
--- a/Engine/Source/Programs/Shared/EpicGames.UHT/Exporters/CodeGen/UhtCodeGenerator.cs
+++ b/Engine/Source/Programs/Shared/EpicGames.UHT/Exporters/CodeGen/UhtCodeGenerator.cs
@@ -510,8 +510,16 @@ namespace EpicGames.UHT.Exporters.CodeGen
 			objectInfo.RegisteredSingletonName = builder.ToString();
 			objectInfo.StaticsName = $"{objectInfo.RegisteredSingletonName}_Statics";
 			objectInfo.RegisteredExternalDecl = $"\t{api}U{obj.EngineLinkClassName}* {objectInfo.RegisteredSingletonName}(ETypeConstructPhase);\r\n";
-			objectInfo.ConstInitSingletonName = objectInfo.RegisteredSingletonName.Replace("Z_Construct_", "Z_ConstInit_", StringComparison.InvariantCultureIgnoreCase);
-			objectInfo.ConstInitExternalDecl = $"\textern UE_CONSTINIT_UOBJECT_DECL TNoDestroy<U{obj.EngineClassName}> {objectInfo.ConstInitSingletonName};\r\n";
+			if (obj is UhtPartial)
+			{
+				objectInfo.ConstInitSingletonName = $"{objectInfo.RegisteredSingletonName}_PartialParams";
+				objectInfo.ConstInitExternalDecl = $"\textern UE_CONSTINIT_UOBJECT_DECL UE::CoreUObject::Private::FPartialClass {objectInfo.ConstInitSingletonName};\r\n";
+			}
+			else
+			{
+				objectInfo.ConstInitSingletonName = objectInfo.RegisteredSingletonName.Replace("Z_Construct_", "Z_ConstInit_", StringComparison.InvariantCultureIgnoreCase);
+				objectInfo.ConstInitExternalDecl = $"\textern UE_CONSTINIT_UOBJECT_DECL TNoDestroy<U{obj.EngineClassName}> {objectInfo.ConstInitSingletonName};\r\n";
+			}
 
 			// Init the children
 			foreach (UhtType child in obj.Children)
diff --git a/Engine/Source/Programs/Shared/EpicGames.UHT/Exporters/CodeGen/UhtHeaderCodeGenerator.cs b/Engine/Source/Programs/Shared/EpicGames.UHT/Exporters/CodeGen/UhtHeaderCodeGenerator.cs
--- a/Engine/Source/Programs/Shared/EpicGames.UHT/Exporters/CodeGen/UhtHeaderCodeGenerator.cs
+++ b/Engine/Source/Programs/Shared/EpicGames.UHT/Exporters/CodeGen/UhtHeaderCodeGenerator.cs
@@ -129,6 +129,11 @@ namespace EpicGames.UHT.Exporters.CodeGen
 			{
 				outerName = classObj.EngineName;
 			}
+			else if (outer is UhtPartial partial)
+			{
+				// For functions in partials, use the extended class name for parameter struct naming
+				outerName = partial.OwnerClass != null ? partial.OwnerClass.EngineName : partial.EngineName;
+			}
 			else if (outer is UhtPackage packageObj)
 			{
 				string packageName = packageObj.EngineName;
@@ -669,6 +674,7 @@ namespace EpicGames.UHT.Exporters.CodeGen
 			{
 				UhtClass classObj => classObj.GeneratedBodyLineNumber,
 				UhtScriptStruct scriptStruct => scriptStruct.MacroDeclaredLineNumber,
+				UhtPartial partial => partial.MacroDeclaredLineNumber,
 				UhtFunction function => function.MacroLineNumber,
 				_ => throw new UhtException(type, "Can not use given type to create a macro"),
 			};
diff --git a/Engine/Source/Programs/Shared/EpicGames.UHT/Exporters/CodeGen/UhtHeaderCodeGeneratorCppFile.cs b/Engine/Source/Programs/Shared/EpicGames.UHT/Exporters/CodeGen/UhtHeaderCodeGeneratorCppFile.cs
--- a/Engine/Source/Programs/Shared/EpicGames.UHT/Exporters/CodeGen/UhtHeaderCodeGeneratorCppFile.cs
+++ b/Engine/Source/Programs/Shared/EpicGames.UHT/Exporters/CodeGen/UhtHeaderCodeGeneratorCppFile.cs
@@ -21,6 +21,7 @@ namespace EpicGames.UHT.Exporters.CodeGen
 		public UhtUsedDefineScopes<UhtScriptStruct> ScriptStructs { get; } = new();
 		public UhtUsedDefineScopes<UhtScriptStruct> RigVMStructs { get; } = new();
 		public UhtUsedDefineScopes<UhtClass> Classes { get; } = new();
+		public UhtUsedDefineScopes<UhtPartial> Partials { get; } = new();
 	}
 
 	internal sealed class UhtHeaderCodeGeneratorCppFile : UhtHeaderCodeGenerator
@@ -94,8 +95,20 @@ namespace EpicGames.UHT.Exporters.CodeGen
 					includesToAdd.Add("UObject/VerseClassProperty.h");
 				}
 
+				// Collect unique partial class names from all partials in this file
+				Dictionary<string, UhtClass> partialClassNames = [];
+
 				foreach (UhtType type in HeaderFile.Children)
 				{
+					if (type is UhtPartial partial && partial.OwnerClass != null)
+					{
+						if (partialClassNames.Count == 0)
+						{
+							includesToAdd.Add("UObject/UObjectPartials.h");
+						}
+						partialClassNames.TryAdd(partial.OwnerClassName, partial.OwnerClass);
+					}
+
 					if (type is UhtStruct structObj)
 					{
 						// Functions
@@ -275,6 +288,10 @@ namespace EpicGames.UHT.Exporters.CodeGen
 						UhtSingletonTypeFlag filter = UhtSingletonTypeFlag.Registered | UhtSingletonTypeFlag.Unregistered | UhtSingletonTypeFlag.ConstInit;
 						foreach (UhtObject obj in HeaderFile.References.InternalObjects.SortReferences(filter))
 						{
+							if (obj is UhtPartial)
+							{
+								continue;
+							}
 							builder.Append(GetExternalDecl(obj, UhtSingletonType.ConstInit).AsSpan().TrimStart());
 						}
 					}
@@ -351,6 +368,13 @@ namespace EpicGames.UHT.Exporters.CodeGen
 							}
 						}
 					}
+					else if (field is UhtPartial partial)
+					{
+						using UhtCodeBlockComment blockComment = new(builder, field);
+						using UhtMacroBlockEmitter macroBlockEmitter = new(builder, partial.DefineScope);
+						AppendPartial(builder, partial);
+						GetRegistrations(packageRegistrations, field).Partials.Add(partial);
+					}
 					else if (field is UhtFunction function)
 					{
 						using UhtCodeBlockComment blockComment = new(builder, field);
@@ -462,6 +486,21 @@ namespace EpicGames.UHT.Exporters.CodeGen
 								combinedHash = IoHash.Combine(combinedHash, hash);
 							},
 							builder => builder.Append("\t};\r\n"));
+
+						builder.AppendInstances(registrations.Partials,
+							builder => builder.Append("\tstatic constexpr FPartialRegisterCompiledInInfo PartialInfo[] = {\r\n"),
+							(builder, partial) =>
+							{
+								IoHash hash = ObjectInfos[partial.ObjectTypeIndex].Hash;
+								builder.Append("\t\t{ ")
+									.Append($"[](ETypeConstructPhase) -> UE::CoreUObject::Private::FPartialClass* {{ return &{GetSingletonName(partial, UhtSingletonType.Registered)}_PartialParams; }},")
+									.Append($"TEXT(\"{partial.EngineName}\"), ")
+									.Append($"&Z_Registration_Info_Partial_{partial.SourceName}, ")
+									.Append($"CONSTRUCT_RELOAD_VERSION_INFO(FPartialReloadVersionInfo, sizeof({partial.FullyQualifiedSourceName}), {(uint)hash}U)")
+									.Append(" },\r\n");
+								combinedHash = IoHash.Combine(combinedHash, hash);
+							},
+							builder => builder.Append("\t};\r\n"));
 					}
 
 					builder.Append($"}}; // UHT_STATICS \r\n");
@@ -488,6 +527,7 @@ namespace EpicGames.UHT.Exporters.CodeGen
 							.AppendArrayPtrAndCountLine(registrations.Classes, UhtNames.StaticsClassInfoId, 1, ",\r\n")
 							.AppendArrayPtrAndCountLine(registrations.ScriptStructs, UhtNames.StaticsScriptStructInfoId, 1, ",\r\n")
 							.AppendArrayPtrAndCountLine(registrations.Enumerations, UhtNames.StaticsEnumInfoId, 1, ",\r\n")
+							.AppendArrayPtrAndCountLine(registrations.Partials, UhtNames.StaticsPartialInfoId, 1, ",\r\n")
 							.Append("};\r\n");
 					}
 				}
@@ -1302,13 +1342,172 @@ namespace EpicGames.UHT.Exporters.CodeGen
 			return builder;
 		}
 
+		private StringBuilder AppendPartial(StringBuilder builder, UhtPartial partial)
+		{
+			// Generate functions for this partial first
+			// Functions must be generated before the function pointer array
+			if (partial.OwnerClass != null)
+			{
+				foreach (UhtFunction function in partial.Functions)
+				{
+					AppendClassFunction(builder, partial.OwnerClass, function);
+				}
+			}
+
+			string registrationName = $"Z_Registration_Info_Partial_{partial.SourceName}";
+
+			if (Session.IsUsingCompiledInObjectFormat(UhtCompiledInObjectFormat.Params))
+			{
+				builder.Append("static FPartialRegistrationInfo ").Append(registrationName).Append(";\r\n");
+			}
+
+			// Collect the properties
+			UhtPropertyMemberContextImpl context = new(this, partial);
+			UhtUsedDefineScopes<UhtProperty> properties = new(partial.Properties);
+
+			using UhtStaticsEmitter statics = new(builder, this, partial);// staticsName);
+
+			builder.Append($"struct UHT_STATICS\r\n");
+			builder.Append("{\r\n");
+
+			// Meta data
+			// No need adding this since the partial "disappears" in runtime so there is way to access the information here
+			builder.AppendConditionalMetaData(partial, UhtNames.TypeId, context, properties, 1);
+
+			// Declare properties
+			// ROBM: NOTE: This should already create property linked lists in the constinit case, but the head of the list is usually a field in a UStruct
+			AppendPropertiesDecl(builder, context, properties, 1);
+
+			if (Session.IsUsingCompiledInObjectFormat(UhtCompiledInObjectFormat.ConstInit) && partial.Functions.Any())
+			{
+				using UhtConditionalMacroBlock block = new(builder, "UE_WITH_CONSTINIT_UOBJECT", Session.IsUsingMultipleCompiledInObjectFormats);
+				builder.Append("\tstatic consteval UFunction* GetFirstFunction()\r\n");
+				builder.Append("\t{\r\n");
+				builder.AppendScopeLink(partial.FirstFunction, context, 2, "return ", ";\r\n");
+				builder.Append("\t}\r\n");
+			}
+
+			builder.Append("};\r\n\r\n");
+
+			// Define properties
+			AppendPropertiesDefs(builder, context, properties, 0);
+			builder.Append("\r\n");
+
+
+			// Wrappers
+			builder.Append($"static void Construct_{partial.SourceName}(void* Memory) {{ new(Memory) {partial.Namespace.FullSourceName}{partial.SourceName}(); }}\r\n");
+			builder.Append($"static void Destruct_{partial.SourceName}(void* Memory) {{ (({partial.Namespace.FullSourceName}{partial.SourceName}*)Memory)->~{partial.SourceName}(); }}\r\n");
+			if (partial.HasBeginDestroy)
+			{
+				// BeginDestroy wrapper (only if partial defines BeginDestroy method)
+				// For now, we'll always generate it and have it check for the method
+				builder.Append($"static void BeginDestroy_{partial.SourceName}(void* Memory) {{ (({partial.Namespace.FullSourceName}{partial.SourceName}*)Memory)->BeginDestroy(); }}\r\n");
+			}
+
+			// Define static PartialOffset member
+			builder.Append($"int32 {partial.Namespace.FullSourceName}{partial.SourceName}::PartialOffset;\r\n\r\n");
+
+			// Forward declare the partial class for partial function compilation
+			if (partial.OwnerClass != null && partial.Functions.Any())
+			{
+				builder.Append($"// Forward declaration for partial function compilation\r\n");
+				builder.Append($"class {partial.OwnerClass.Namespace.FullSourceName}{partial.OwnerClass.SourceName};\r\n\r\n");
+			}
+
+			// Generate function arrays for partial functions (after functions are generated)
+			if (!Session.IsUsingCompiledInObjectFormat(UhtCompiledInObjectFormat.ConstInit) && partial.Functions.Any())
+			{
+				// Function constructor pointers (for creating UFunctions)
+				builder.Append($"// Function constructor pointers for {partial.SourceName}\r\n");
+				builder.Append($"static UFunction* (*PartialFunctionConstructors_{partial.SourceName}[])(ETypeConstructPhase) = {{\r\n");
+				foreach (UhtFunction function in partial.Functions)
+				{
+					builder.Append($"\t&{GetSingletonName(function, UhtSingletonType.Registered)},\r\n");
+				}
+				builder.Append("};\r\n\r\n");
+
+				// Native function bindings (for registering exec functions)
+				builder.Append($"// Native function bindings for {partial.SourceName}\r\n");
+				builder.Append($"static constexpr UE::CodeGen::FClassNativeFunction PartialNativeFunctions_{partial.SourceName}[] = {{\r\n");
+				foreach (UhtFunction function in partial.Functions)
+				{
+					builder.Append($"\t{{ .NameUTF8 = UTF8TEXT(\"{function.EngineName}\"), .Pointer = &{partial.SourceName}::{function.UnMarshalAndCallName} }},\r\n");
+				}
+				builder.Append("};\r\n\r\n");
+			}
+
+			// ROBM: There should be an if/else constinit case here, and the constinit case should just store the head of the linked list of properties etc
+			// See UClass/UStruct versions 
+			// There should be a global var like `Z_ConstInit_UPartial_{BaseName}_{PartialName}` here that can be linked in the package code generator 
+			builder.Append($"UE_CONSTINIT_UOBJECT_DECL UE::CoreUObject::Private::FPartialClass {GetSingletonName(partial, UhtSingletonType.Registered)}_PartialParams\r\n");
+			builder.Append("{\r\n");
+			builder.Append($"\t.Name = TEXT(\"{partial.SourceName}\"),\r\n");
+			//builder.Append($"\t.ClassName = TEXT(\"{partial.OwnerClassName}\"),\r\n");
+			builder.Append($"\t.Size = sizeof({partial.Namespace.FullSourceName}{partial.SourceName}),\r\n");
+			builder.Append($"\t.Alignment = alignof({partial.Namespace.FullSourceName}{partial.SourceName}),\r\n");
+
+			if (Session.IsUsingCompiledInObjectFormat(UhtCompiledInObjectFormat.ConstInit))
+			{
+				using UhtConditionalMacroBlock block = new(builder, "UE_WITH_CONSTINIT_UOBJECT", Session.IsUsingMultipleCompiledInObjectFormats);
+				builder.Append($"\t.Class = &{GetSingletonName(partial.OwnerClass, UhtSingletonType.ConstInit)},\r\n");
+
+				if (partial.Properties.Any())
+				{
+					builder.Append($"\t.FirstProperty = UHT_STATICS::GetFirstProperty(),\r\n");
+				}
+				if (partial.Functions.Any())
+				{
+					builder.Append($"\t.FirstChild = UHT_STATICS::GetFirstFunction(),\r\n");
+				}
+			}
+
+			if (Session.IsUsingCompiledInObjectFormat(UhtCompiledInObjectFormat.Params))
+			{
+				using UhtConditionalMacroBlock block = new(builder, "!UE_WITH_CONSTINIT_UOBJECT", Session.IsUsingMultipleCompiledInObjectFormats);
+				builder.Append($"\t.Class = {GetSingletonName(partial.OwnerClass, UhtSingletonType.Registered)},\r\n");
+				if (partial.Properties.Any())
+				{
+					builder.Append($"\t.Properties = UHT_STATICS::PropPointers,\r\n");
+					builder.Append($"\t.NumProperties = UE_ARRAY_COUNT(UHT_STATICS::PropPointers),\r\n");
+				}
+				if (partial.Functions.Any())
+				{
+					builder.Append($"\t.FunctionConstructors = PartialFunctionConstructors_{partial.SourceName},\r\n");
+					builder.Append($"\t.NativeFunctions = PartialNativeFunctions_{partial.SourceName},\r\n");
+					builder.Append($"\t.NumFunctions = {partial.Functions.Count()},\r\n");
+				}
+			}
+
+			builder.Append($"\t.PartialOffsetPtr = &{partial.Namespace.FullSourceName}{partial.SourceName}::PartialOffset,\r\n");
+			builder.Append($"\t.Constructor = &Construct_{partial.SourceName},\r\n");
+			builder.Append($"\t.BeginDestroy = {(partial.HasBeginDestroy ? $"BeginDestroy_{partial.SourceName}" : "nullptr")},\r\n");
+			builder.Append($"\t.Destructor = &Destruct_{partial.SourceName},\r\n");
+			builder.Append($"\t.GetLifetimeReplicatedProps = nullptr, // TODO: Detect if method exists\r\n");
+			builder.Append("};\r\n");
+
+			return builder;
+		}
+
 		private StringBuilder AppendFunction(StringBuilder builder, UhtFunction function, bool isNoExport, bool isVerseNativeCallable)
 		{
 			string singletonName = GetSingletonName(function, UhtSingletonType.Registered);
 			string constinitName = GetSingletonName(function, UhtSingletonType.ConstInit);
 			bool paramsInStatic = isNoExport || !IsCallbackFunction(function);
 			bool isNet = function.FunctionFlags.HasAnyFlags(EFunctionFlags.NetRequest | EFunctionFlags.NetResponse);
-			UhtClass? outerClass = function.Outer as UhtClass;
+
+			string fullSourceName = "";
+			string callName = "";
+
+			// Get the outer class - either directly or from the partial's partial class
+			if (function.Outer is UhtClass outerClass)
+			{
+				fullSourceName = outerClass.Namespace.FullSourceName;
+				callName = outerClass.NativeFunctionCallName;
+			}
+			else if (function.Outer is UhtPartial partial)
+			{
+				fullSourceName = partial.SourceName;
+			}
 
 			using UhtStaticsEmitter statics = new(builder, this, function);
 
@@ -1408,7 +1607,7 @@ namespace EpicGames.UHT.Exporters.CodeGen
 								=> $"\t\t.NativeFunction = UHT_STATICS::GetCustomThunk_{function.SourceName}(),\r\n",
 							// Native without custom thunk
 							(EFunctionFlags.Native, EFunctionFlags.None, _)
-								=> $"\t\t.NativeFunction = &{outerClass!.Namespace.FullSourceName}{outerClass!.NativeFunctionCallName}::exec{function.SourceName},\r\n",
+								=> $"\t\t.NativeFunction = &{fullSourceName}{callName}::exec{function.SourceName},\r\n",
 							// Non-native function 
 							// Can't use cross-dll assignment in modular builds here 
 							(EFunctionFlags.None, _, _) => ".NativeFunction = UE_IF(IS_MONOLITHIC, &UObject::ProcessInternal, nullptr),\r\n",
@@ -1518,6 +1717,15 @@ namespace EpicGames.UHT.Exporters.CodeGen
 			{
 				return GetSingletonName(package, UhtSingletonType.Registered);
 			}
+			else if (function.Outer is UhtPartial partial)
+			{
+				// Functions in partials should use the partial class as their outer
+				if (partial.OwnerClass != null)
+				{
+					return GetSingletonName(partial.OwnerClass, UhtSingletonType.Registered);
+				}
+				return "nullptr";
+			}
 			else
 			{
 				return GetSingletonName((UhtObject)function.Outer, UhtSingletonType.Registered);
@@ -1758,10 +1966,21 @@ namespace EpicGames.UHT.Exporters.CodeGen
 				}
 				if (isRpc)
 				{
-					builder.Append($"DEFINE_FUNCTION({classObj.Namespace.FullSourceName}{classObj.NativeFunctionCallName}::{function.UnMarshalAndCallName})\r\n");
-					builder.Append("{\r\n");
-					AppendFunctionThunk(builder, classObj, function);
-					builder.Append("}\r\n");
+					if (function.Outer is UhtPartial partial)
+					{
+						builder.Append($"DEFINE_FUNCTION({partial.SourceName}::{function.UnMarshalAndCallName})\r\n");
+						builder.Append("{\r\n");
+						builder.Append($"\ttypedef {classObj.Namespace.FullSourceName}{classObj.SourceName} ThisClass;\r\n");
+						AppendFunctionThunk(builder, classObj, function);
+						builder.Append("}\r\n");
+					}
+					else
+					{
+						builder.Append($"DEFINE_FUNCTION({classObj.Namespace.FullSourceName}{classObj.NativeFunctionCallName}::{function.UnMarshalAndCallName})\r\n");
+						builder.Append("{\r\n");
+						AppendFunctionThunk(builder, classObj, function);
+						builder.Append("}\r\n");
+					}
 				}
 			}
 			return builder;
@@ -3355,6 +3574,11 @@ namespace EpicGames.UHT.Exporters.CodeGen
 					builder.Append($"{classObj.SourceName}::{function.CppImplName}(").AppendFunctionThunkParameterNames(function).Append(");\r\n");
 				}
 			}
+			else if (function.Outer is UhtPartial partial)
+			{
+				// Function is in an partial - call through GetPartial<>()
+				builder.Append("P_THIS->GetPartial<").Append(partial.SourceName).Append(">().").Append(function.CppImplName).Append('(').AppendFunctionThunkParameterNames(function).Append(");\r\n");
+			}
 			else
 			{
 				builder.Append("P_THIS->").Append(function.CppImplName).Append('(').AppendFunctionThunkParameterNames(function).Append(");\r\n");
@@ -3452,7 +3676,7 @@ namespace EpicGames.UHT.Exporters.CodeGen
 			{
 				_codeGenerator = codeGenerator;
 				_outerStruct = outerStruct;
-				string outerStructNamespaceName = _outerStruct is UhtScriptStruct scriptStruct ? scriptStruct.NamespaceExportName : _outerStruct.Namespace.FullSourceName;
+				string outerStructNamespaceName = _outerStruct is UhtScriptStruct scriptStruct ? scriptStruct.NamespaceExportName : (_outerStruct is UhtPartial partial ? partial.NamespaceExportName : _outerStruct.Namespace.FullSourceName);
 				string outerStructSourceName = _outerStruct.SourceName;
 				_outerIdentifier = new(outerStructNamespaceName, null, outerStructSourceName, null);
 			}
@@ -3514,8 +3738,13 @@ namespace EpicGames.UHT.Exporters.CodeGen
 			{
 				builder.AppendLine($"{Tabs(tabs + 1)}.Class = {SingletonRef(linker, uhtObject.EngineClass, UhtSingletonType.ConstInit)}, ");
 			}
+			UhtObject? outer = uhtObject.Outer as UhtObject;
+			if (outer is UhtPartial partial)
+			{
+				outer = partial.OwnerClass;
+			}
 			builder.AppendTabs(tabs + 1).Append($".NameUTF8 = UTF8TEXT(").AppendUTF8LiteralString(uhtObject.EngineName).Append("),\r\n");
-			builder.AppendTabs(tabs + 1).AppendLine($".Outer = {ConstInitSingletonRef(linker, uhtObject.Outer as UhtObject)},");
+			builder.AppendTabs(tabs + 1).AppendLine($".Outer = {ConstInitSingletonRef(linker, outer)},");
 			builder.AppendTabs(tabs).Append("},\r\n");
 			return builder;
 		}
diff --git a/Engine/Source/Programs/Shared/EpicGames.UHT/Exporters/CodeGen/UhtHeaderCodeGeneratorHFile.cs b/Engine/Source/Programs/Shared/EpicGames.UHT/Exporters/CodeGen/UhtHeaderCodeGeneratorHFile.cs
--- a/Engine/Source/Programs/Shared/EpicGames.UHT/Exporters/CodeGen/UhtHeaderCodeGeneratorHFile.cs
+++ b/Engine/Source/Programs/Shared/EpicGames.UHT/Exporters/CodeGen/UhtHeaderCodeGeneratorHFile.cs
@@ -181,6 +181,12 @@ namespace EpicGames.UHT.Exporters.CodeGen
 							AppendScriptStruct(builder, scriptStruct);
 						}
 					}
+					else if (field is UhtPartial partial)
+					{
+						using UhtCodeBlockComment blockComment = new(builder, field);
+						using UhtMacroBlockEmitter macroBlockEmitter = new(builder, partial.DefineScope);
+						Appendpartial(builder, partial);
+					}
 					else if (field is UhtFunction function)
 					{
 					}
@@ -730,6 +736,31 @@ namespace EpicGames.UHT.Exporters.CodeGen
 			return builder;
 		}
 
+		private StringBuilder Appendpartial(StringBuilder builder, UhtPartial partial)
+		{
+			using (UhtMacroCreator macro = new(builder, this, partial, GeneratedBodyMacroSuffix))
+			{
+				builder.Append('\t').Append(Module.Api).Append("static int32 PartialOffset; \\\r\n");
+				builder.Append('\t').Append(partial.OwnerClassName).Append("& GetOwner() {");
+				builder.Append(" return *reinterpret_cast<").Append(partial.OwnerClassName).Append("*>((uint8*)this - PartialOffset); } \\\r\n");
+				builder.Append("\tconst ").Append(partial.OwnerClassName).Append("& GetOwner() const {");
+				builder.Append(" return *reinterpret_cast<const ").Append(partial.OwnerClassName).Append("*>((const uint8*)this - PartialOffset); } \\\r\n");
+				if (partial.Functions.Any())
+				{
+					builder.Append("\t \\\r\n");
+					foreach (UhtFunction function in partial.Functions)
+					{
+						if (function.FunctionFlags.HasAnyFlags(EFunctionFlags.Native))
+						{
+							builder.Append("\tDECLARE_FUNCTION(").Append(function.UnMarshalAndCallName).Append("); \\\r\n");
+						}
+					}
+				}
+			}
+
+			return builder;
+		}
+
 		private StringBuilder AppendEnum(StringBuilder builder, UhtEnum enumObj)
 		{
 			// Export FOREACH macro
@@ -1682,6 +1713,14 @@ namespace EpicGames.UHT.Exporters.CodeGen
 					builder.Append("\tDECLARE_WITHIN(").Append(classObj.ClassWithin.SourceName).Append(") \\\r\n");
 				}
 			}
+
+			// Only generate GetPartial<>() if parent is not partial (i.e., this is the first partial class in the hierarchy)
+			if (classObj.ClassFlags.HasFlag(EClassFlags.Partial) && !classObj.SuperClass!.ClassFlags.HasFlag(EClassFlags.Partial))
+			{
+				builder.Append("\ttemplate<typename T> T& GetPartial() { return *reinterpret_cast<T*>((uint8*)this + T::PartialOffset); } \\\r\n");
+				builder.Append("\ttemplate<typename T> const T& GetPartial() const { return *reinterpret_cast<const T*>((const uint8*)this + T::PartialOffset); } \\\r\n");
+			}
+
 			return builder;
 		}
 
@@ -1702,6 +1741,10 @@ namespace EpicGames.UHT.Exporters.CodeGen
 			{
 				builder.Append(" | CLASS_Optional");
 			}
+			if (classObj.ClassFlags.HasAnyFlags(EClassFlags.Partial))
+			{
+				builder.Append(" | CLASS_Partial");
+			}
 			if (classObj.ClassFlags.HasAnyFlags(EClassFlags.DefaultConfig))
 			{
 				builder.Append(" | CLASS_DefaultConfig");
diff --git a/Engine/Source/Programs/Shared/EpicGames.UHT/Exporters/CodeGen/UhtPackageCodeGeneratorCppFile.cs b/Engine/Source/Programs/Shared/EpicGames.UHT/Exporters/CodeGen/UhtPackageCodeGeneratorCppFile.cs
--- a/Engine/Source/Programs/Shared/EpicGames.UHT/Exporters/CodeGen/UhtPackageCodeGeneratorCppFile.cs
+++ b/Engine/Source/Programs/Shared/EpicGames.UHT/Exporters/CodeGen/UhtPackageCodeGeneratorCppFile.cs
@@ -144,6 +144,9 @@ namespace EpicGames.UHT.Exporters.CodeGen
 					.Select((int order, UhtObject obj) (p) => (p.Value, p.Key))
 					.OrderBy(p => p.order)
 					.Select(p => p.obj)];
+
+				List<UhtPartial> partials = Module.Packages.SelectMany(p => p.Children.OfType<UhtPartial>().Where(o => o.AlternateObject is null)).ToList();
+
 				if (Session.IsUsingCompiledInObjectFormat(UhtCompiledInObjectFormat.ConstInit))
 				{
 					using UhtConditionalMacroBlock block = new(builder, "UE_WITH_CONSTINIT_UOBJECT", Session.IsUsingMultipleCompiledInObjectFormats);
@@ -156,6 +159,11 @@ namespace EpicGames.UHT.Exporters.CodeGen
 						builder.Append(GetExternalDecl(obj, UhtSingletonType.ConstInit));
 					}
 
+					foreach (UhtPartial partial in partials)
+					{
+						builder.Append(GetExternalDecl(partial, UhtSingletonType.ConstInit));
+					}
+
 					builder.AppendLine($$"""
 						extern {{Module.Api}}UObject* const {{ModuleLinkableObjects(Module)}}[{{linkableObjectsOrdered.Count}}] = 
 						{
@@ -175,6 +183,28 @@ namespace EpicGames.UHT.Exporters.CodeGen
 
 						""");
 
+					if (partials.Count != 0)
+					{
+						builder.AppendLine($$"""
+							extern UE::CoreUObject::Private::FPartialClass* const Z_Partials_{{Module.ShortName}}[{{partials.Count}}] = 
+							{
+							""");
+						foreach (UhtPartial partial in partials)
+						{
+							using UhtMacroBlockEmitter macroBlockEmitter = new(builder, partial.DefineScope);
+							builder.AppendLine($"\t{SingletonRef(this, partial, UhtSingletonType.ConstInit)}, ");
+							if (partial.DefineScope != UhtDefineScope.None)
+							{
+								macroBlockEmitter.SetElse(UhtDefineScope.None);
+								builder.AppendLine("\tnullptr, ");
+							}
+						}
+						builder.AppendLine($$"""
+							};	
+
+							""");
+					}
+
 					List<UhtModule> orderedLinkedModules = [.. Module.Linking!.ReferencedModules
 						.Select((int Index, UhtModule M) (p) => (p.Value, p.Key))
 						.OrderBy(p => p.Index)
@@ -297,8 +327,12 @@ namespace EpicGames.UHT.Exporters.CodeGen
 							.Structs = MakeConstArrayView({{ModuleLinkableObjects(Module)}}).Mid({{structsStart}}, {{structsEnd - structsStart}}),
 							.Enums = MakeConstArrayView({{ModuleLinkableObjects(Module)}}).Mid({{enumsStart}}, {{enumsEnd - enumsStart}}),
 							.Others = MakeConstArrayView({{ModuleLinkableObjects(Module)}}).RightChop({{othersStart}}),
-						} };
 						""");
+					if (partials.Count > 0)
+					{
+						builder.Append($"\t.Partials = MakeConstArrayView(Z_Partials_{Module.ShortName}),\r\n");
+					}
+					builder.AppendLine("} };");
 				}
 
 				// Verse registration
diff --git a/Engine/Source/Programs/Shared/EpicGames.UHT/Parsers/UhtFunctionParser.cs b/Engine/Source/Programs/Shared/EpicGames.UHT/Parsers/UhtFunctionParser.cs
--- a/Engine/Source/Programs/Shared/EpicGames.UHT/Parsers/UhtFunctionParser.cs
+++ b/Engine/Source/Programs/Shared/EpicGames.UHT/Parsers/UhtFunctionParser.cs
@@ -109,6 +109,7 @@ namespace EpicGames.UHT.Parsers
 		[UhtKeyword(Extends = UhtTableNames.Interface)]
 		[UhtKeyword(Extends = UhtTableNames.NativeInterface)]
 		[UhtKeyword(Extends = UhtTableNames.VModule)]
+		[UhtKeyword(Extends = UhtTableNames.Partial)]
 		[SuppressMessage("CodeQuality", "IDE0051:Remove unused private members", Justification = "Attribute accessed method")]
 		[SuppressMessage("Style", "IDE0060:Remove unused parameter", Justification = "Attribute accessed method")]
 		private static UhtParseResult UFUNCTIONKeyword(UhtParsingScope topScope, UhtParsingScope actionScope, ref UhtToken token)
@@ -313,7 +314,22 @@ namespace EpicGames.UHT.Parsers
 			{
 				using UhtParsingScope topScope = new(parentScope, function, parentScope.Session.GetKeywordTable(UhtTableNames.Function), UhtAccessSpecifier.Public);
 				UhtParsingScope outerClassScope = topScope.CurrentClassScope;
-				UhtClass outerClass = (UhtClass)outerClassScope.ScopeType;
+
+				// Get the outer class - either directly or from partial's owner class
+				UhtClass? outerClass = outerClassScope.ScopeType as UhtClass;
+				UhtPartial? outerPartial = null;
+				if (outerClass == null && outerClassScope.ScopeType is UhtPartial partial)
+				{
+					outerPartial = partial;
+					outerClass = partial.OwnerClass;
+				}
+
+				// If we're not in a class or partial, that's an error
+				if (outerClass == null && outerPartial == null)
+				{
+					throw new UhtException(function, "UFUNCTIONs must be declared in a class or partial");
+				}
+
 				string scopeName = "function";
 
 				{
@@ -324,7 +340,8 @@ namespace EpicGames.UHT.Parsers
 					UhtSpecifierParser specifierParser = UhtSpecifierParser.GetThreadInstance(specifierContext, scopeName, parentScope.Session.GetSpecifierTable(UhtTableNames.Function));
 					specifierParser.ParseSpecifiers();
 
-					if (!outerClass.ClassFlags.HasAnyFlags(EClassFlags.Native))
+					// Only check native class flag if we have a resolved outer class (not an partial during parsing)
+					if (outerClass != null && !outerClass.ClassFlags.HasAnyFlags(EClassFlags.Native))
 					{
 						throw new UhtException(function, "Should only be here for native classes!");
 					}
diff --git a/Engine/Source/Programs/Shared/EpicGames.UHT/Parsers/UhtHeaderFileParser.cs b/Engine/Source/Programs/Shared/EpicGames.UHT/Parsers/UhtHeaderFileParser.cs
--- a/Engine/Source/Programs/Shared/EpicGames.UHT/Parsers/UhtHeaderFileParser.cs
+++ b/Engine/Source/Programs/Shared/EpicGames.UHT/Parsers/UhtHeaderFileParser.cs
@@ -694,6 +694,22 @@ namespace EpicGames.UHT.Parsers
 						}
 					}
 				}
+				else if (topScope.ScopeType is UhtPartial partial)
+				{
+					using UhtTokenRecorder recorder = new(topScope, ref token);
+					topScope.SkipDeclaration(token);
+					if (recorder.Stop())
+					{
+						if (partial.Declarations != null)
+						{
+							UhtDeclaration declaration = partial.Declarations[partial.Declarations.Count - 1];
+							if (topScope.HeaderParser.CheckForBeginDestroy(partial, declaration))
+							{
+							}
+							// Add AddReferencedObjects, PostLoad, GetLifetimeReplicatedProps and friends later for feature complete
+						}
+					}
+				}
 				else
 				{
 					topScope.SkipDeclaration(token);
@@ -1482,6 +1498,24 @@ namespace EpicGames.UHT.Parsers
 
 			return false;
 		}
+
+		private bool CheckForBeginDestroy(UhtPartial partial, UhtDeclaration declaration)
+		{
+			using UhtTokenReplayReaderBorrower borrowedReader = new(HeaderFile, HeaderFile.Data.Memory, declaration.Tokens, UhtTokenType.EndOfDeclaration);
+			IUhtTokenReader replayReader = borrowedReader.Reader;
+
+			if (!replayReader.TryOptional("void") ||
+				!replayReader.TryOptional("BeginDestroy") ||
+				!replayReader.TryOptional('('))
+			{
+				return false;
+			}
+
+			partial.HasBeginDestroy = true;
+
+			return true;
+		}
+
 		#endregion
 	}
 }
diff --git a/Engine/Source/Programs/Shared/EpicGames.UHT/Parsers/UhtParsingScope.cs b/Engine/Source/Programs/Shared/EpicGames.UHT/Parsers/UhtParsingScope.cs
--- a/Engine/Source/Programs/Shared/EpicGames.UHT/Parsers/UhtParsingScope.cs
+++ b/Engine/Source/Programs/Shared/EpicGames.UHT/Parsers/UhtParsingScope.cs
@@ -73,7 +73,7 @@ namespace EpicGames.UHT.Parsers
 				UhtParsingScope? currentScope = this;
 				while (currentScope != null)
 				{
-					if (currentScope.ScopeType is UhtClass)
+					if (currentScope.ScopeType is UhtClass or UhtPartial)
 					{
 						return currentScope;
 					}
@@ -623,7 +623,7 @@ namespace EpicGames.UHT.Parsers
 			_compilerDirective = _scope.HeaderParser.GetCurrentCompositeCompilerDirective();
 			_function = null;
 			_flushed = false;
-			if (_scope.ScopeType is UhtClass)
+			if (_scope.ScopeType is UhtClass or UhtPartial)
 			{
 				_scope.TokenReader.EnableRecording();
 				_scope.TokenReader.RecordToken(ref initialToken);
@@ -641,7 +641,7 @@ namespace EpicGames.UHT.Parsers
 			_compilerDirective = _scope.HeaderParser.GetCurrentCompositeCompilerDirective();
 			_function = function;
 			_flushed = false;
-			if (_scope.ScopeType is UhtClass)
+			if (_scope.ScopeType is UhtClass or UhtPartial)
 			{
 				_scope.TokenReader.EnableRecording();
 			}
@@ -658,7 +658,7 @@ namespace EpicGames.UHT.Parsers
 		/// <summary>
 		/// Stop the recording
 		/// </summary>
-		/// <returns>True if the recorded content was added to a class</returns>
+		/// <returns>True if the recorded content was added to a class or partial</returns>
 		public bool Stop()
 		{
 			if (!_flushed)
@@ -670,6 +670,12 @@ namespace EpicGames.UHT.Parsers
 					_scope.TokenReader.DisableRecording();
 					return true;
 				}
+				else if (_scope.ScopeType is UhtPartial partialObj)
+				{
+					partialObj.AddDeclaration(_compilerDirective, _scope.TokenReader.RecordedTokens, _function);
+					_scope.TokenReader.DisableRecording();
+					return true;
+				}
 			}
 			return false;
 		}
diff --git a/Engine/Source/Programs/Shared/EpicGames.UHT/Parsers/UhtPartialParser.cs b/Engine/Source/Programs/Shared/EpicGames.UHT/Parsers/UhtPartialParser.cs
new file mode 100644
--- /dev/null
+++ b/Engine/Source/Programs/Shared/EpicGames.UHT/Parsers/UhtPartialParser.cs
@@ -0,0 +1,135 @@
+// Copyright Epic Games, Inc. All Rights Reserved.
+
+using System.Diagnostics.CodeAnalysis;
+using EpicGames.UHT.Tables;
+using EpicGames.UHT.Tokenizer;
+using EpicGames.UHT.Types;
+using EpicGames.UHT.Utils;
+
+namespace EpicGames.UHT.Parsers
+{
+	/// <summary>
+	/// UPARTIAL parser object
+	/// </summary>
+	[UnrealHeaderTool]
+	public static class UhtPartialParser
+	{
+		#region Keywords
+		[UhtKeyword(Extends = UhtTableNames.Global)]
+		[SuppressMessage("CodeQuality", "IDE0051:Remove unused private members", Justification = "Attribute accessed method")]
+		[SuppressMessage("Style", "IDE0060:Remove unused parameter", Justification = "Attribute accessed method")]
+		private static UhtParseResult UPARTIALKeyword(UhtParsingScope topScope, UhtParsingScope actionScope, ref UhtToken token)
+		{
+			return ParseUPartial(topScope, token);
+		}
+
+		[UhtKeyword(Extends = UhtTableNames.Partial)]
+		[SuppressMessage("CodeQuality", "IDE0051:Remove unused private members", Justification = "Attribute accessed method")]
+		[SuppressMessage("Style", "IDE0060:Remove unused parameter", Justification = "Attribute accessed method")]
+		private static UhtParseResult GENERATED_BODYKeyword(UhtParsingScope topScope, UhtParsingScope actionScope, ref UhtToken token)
+		{
+			UhtPartial partial = (UhtPartial)topScope.ScopeType;
+			if (topScope.AccessSpecifier != UhtAccessSpecifier.Public)
+			{
+				topScope.TokenReader.LogError($"{token.Value} must be in the public scope of '{partial.SourceName}', not private or protected.");
+			}
+
+			if (partial.MacroDeclaredLineNumber != -1)
+			{
+				topScope.TokenReader.LogError($"Multiple {token.Value} declarations found in '{partial.SourceName}'");
+			}
+
+			partial.MacroDeclaredLineNumber = topScope.TokenReader.InputLine;
+
+			topScope.TokenReader
+				.Require('(')
+				.Require(')')
+				.OptionalAttributes(false);
+			return UhtParseResult.Handled;
+		}
+		#endregion
+
+		private static UhtParseResult ParseUPartial(UhtParsingScope parentScope, UhtToken keywordToken)
+		{
+			UhtPartial partial = new(parentScope.HeaderFile, parentScope.HeaderParser.GetNamespace(), parentScope.ScopeType, keywordToken.InputLine);
+
+			// Use Partial keyword table to allow UFUNCTIONs without class-specific specifiers
+			using UhtParsingScope topScope = new(parentScope, partial, parentScope.Session.GetKeywordTable(UhtTableNames.Partial), UhtAccessSpecifier.Public);
+
+			const string ScopeName = "partial";
+			using UhtMessageContext tokenContext = new(ScopeName);
+
+			// Parse UPARTIAL(ClassName) - the owner class name is required in parentheses
+			topScope.TokenReader
+				.Require('(')
+				.RequireIdentifier((ref UhtToken classToken) => partial.OwnerClassName = classToken.Value.ToString())
+				.Require(')');
+
+			// No specifiers for partials (unlike USTRUCT)
+			// If we want to support specifiers later, we can add them here
+
+			// Consume the struct keyword
+			topScope.TokenReader.OptionalAttributes(false);
+			topScope.TokenReader.Require("struct");
+			topScope.TokenReader.OptionalAttributes(true);
+
+			// Read the partial struct name and possible API macro name
+			topScope.TokenReader.TryOptionalAPIMacro(out UhtToken apiMacroToken);
+			partial.SourceName = topScope.TokenReader.GetIdentifier().Value.ToString();
+
+			topScope.AddModuleRelativePathToMetaData();
+
+			// Strip the name (partials must start with F)
+			if (partial.SourceName[0] == 'F')
+			{
+				partial.EngineName = partial.SourceName[1..];
+			}
+			else
+			{
+				// This will be flagged later in the validation phase
+				partial.EngineName = partial.SourceName;
+				topScope.TokenReader.LogError($"partial struct '{partial.SourceName}' must start with 'F' prefix");
+			}
+
+			// Check for an empty engine name
+			if (partial.EngineName.Length == 0)
+			{
+				topScope.TokenReader.LogError($"When compiling partial definition for '{partial.SourceName}', attempting to strip prefix results in an empty name. Did you leave off a prefix?");
+			}
+
+			// Partials cannot have inheritance
+			topScope.TokenReader.Optional("final");
+
+			// Check for inheritance - partials cannot inherit
+			if (topScope.TokenReader.TryOptional(':'))
+			{
+				topScope.TokenReader.LogError($"Partial '{partial.SourceName}' cannot have inheritance. Partials cannot inherit from other structs or partials.");
+			}
+
+			// Add the comments here for compatibility with old UHT
+			topScope.TokenReader.PeekToken();
+			topScope.TokenReader.CommitPendingComments();
+			topScope.AddFormattedCommentsAsTooltipMetaData();
+
+			// Add metadata to indicate this is an partial
+			partial.MetaData.Add("IsPartial", "true");
+			partial.MetaData.Add("OwnerClass", partial.OwnerClassName);
+
+			UhtCompilerDirective compilerDirective = topScope.HeaderParser.GetCurrentCompositeCompilerDirective();
+			partial.DefineScope |= compilerDirective.GetDefaultDefineScopes();
+
+			partial.Outer?.AddChild(partial);
+
+			topScope.HeaderParser.ParseStatements('{', '}', true);
+
+			topScope.TokenReader.Require(';');
+
+			if (partial.MacroDeclaredLineNumber == -1)
+			{
+				topScope.TokenReader.LogError("Expected a GENERATED_BODY() at the start of the partial");
+			}
+
+			return UhtParseResult.Handled;
+		}
+	}
+}
diff --git a/Engine/Source/Programs/Shared/EpicGames.UHT/Specifiers/UhtClassSpecifiers.cs b/Engine/Source/Programs/Shared/EpicGames.UHT/Specifiers/UhtClassSpecifiers.cs
--- a/Engine/Source/Programs/Shared/EpicGames.UHT/Specifiers/UhtClassSpecifiers.cs
+++ b/Engine/Source/Programs/Shared/EpicGames.UHT/Specifiers/UhtClassSpecifiers.cs
@@ -113,6 +113,13 @@ namespace EpicGames.UHT.Parsers
 			classObj.AddClassFlags(EClassFlags.Const);
 		}
 
+		[UhtSpecifier(Extends = UhtTableNames.Class, ValueType = UhtSpecifierValueType.None)]
+		private static void PartialSpecifier(UhtSpecifierContext specifierContext)
+		{
+			UhtClass classObj = (UhtClass)specifierContext.Type;
+			classObj.AddClassFlags(EClassFlags.Partial);
+		}
+
 		[UhtSpecifier(Extends = UhtTableNames.Class, ValueType = UhtSpecifierValueType.None)]
 		private static void PerObjectConfigSpecifier(UhtSpecifierContext specifierContext)
 		{
diff --git a/Engine/Source/Programs/Shared/EpicGames.UHT/Types/UhtFunction.cs b/Engine/Source/Programs/Shared/EpicGames.UHT/Types/UhtFunction.cs
--- a/Engine/Source/Programs/Shared/EpicGames.UHT/Types/UhtFunction.cs
+++ b/Engine/Source/Programs/Shared/EpicGames.UHT/Types/UhtFunction.cs
@@ -760,7 +760,13 @@ namespace EpicGames.UHT.Types
 			switch (FunctionType)
 			{
 				case UhtFunctionType.Function:
-					UhtClass? outerClass = (UhtClass?)Outer;
+
+					// Get the outer class - either directly or from partial's Owner class
+					UhtClass? outerClass = Outer as UhtClass;
+					if (outerClass == null && Outer is UhtPartial partial)
+					{
+						outerClass = partial.OwnerClass;
+					}
 
 					if (FunctionFlags.HasAnyFlags(EFunctionFlags.BlueprintPure) && outerClass != null &&
 						outerClass.ClassFlags.HasAnyFlags(EClassFlags.Interface) &&
@@ -1194,7 +1200,12 @@ namespace EpicGames.UHT.Types
 				return;
 			}
 
-			UhtClass? outerClass = (UhtClass?)Outer;
+			// Get the outer class - either directly or from partial's Owner class
+			UhtClass? outerClass = Outer as UhtClass;
+			if (outerClass == null && Outer is UhtPartial partial)
+			{
+				outerClass = partial.OwnerClass;
+			}
 
 			if (policy.FunctionToolTipsRequired)
 			{
diff --git a/Engine/Source/Programs/Shared/EpicGames.UHT/Types/UhtPartial.cs b/Engine/Source/Programs/Shared/EpicGames.UHT/Types/UhtPartial.cs
new file mode 100644
--- /dev/null
+++ b/Engine/Source/Programs/Shared/EpicGames.UHT/Types/UhtPartial.cs
@@ -0,0 +1,287 @@
+// Copyright Epic Games, Inc. All Rights Reserved.
+
+using System;
+using System.Collections.Generic;
+using System.Text;
+using System.Text.Json.Serialization;
+using EpicGames.Core;
+using EpicGames.UHT.Exporters.CodeGen;
+using EpicGames.UHT.Parsers;
+using EpicGames.UHT.Tables;
+using EpicGames.UHT.Tokenizer;
+using EpicGames.UHT.Utils;
+
+namespace EpicGames.UHT.Types
+{
+	/// <summary>
+	/// Represents an partial - a struct that extends a UObject-derived class with additional properties.
+	/// Partials are similar to USTRUCTs but their properties are merged directly into the owner class.
+	/// </summary>
+	public class UhtPartial : UhtStruct
+	{
+		private List<UhtDeclaration>? _declarations = null;
+
+		/// <summary>
+		/// Name of the class that this partial extends (e.g., "APlayerController")
+		/// </summary>
+		public string OwnerClassName { get; set; } = String.Empty;
+
+		/// <summary>
+		/// Reference to the partial class (resolved during resolution phase)
+		/// </summary>
+		[JsonConverter(typeof(UhtNullableTypeSourceNameJsonConverter<UhtClass>))]
+		public UhtClass? OwnerClass { get; set; } = null;
+
+		/// <summary>
+		/// Line number where GENERATED_BODY macro was found
+		/// </summary>
+		public int MacroDeclaredLineNumber { get; set; } = -1;
+
+		/// <inheritdoc/>
+		[JsonIgnore]
+		public override UhtEngineType EngineType => UhtEngineType.ScriptStruct;
+
+		/// <inheritdoc/>
+		public override string EngineClassName => "Partial";
+
+		/// <inheritdoc/>
+		public override string EngineLinkClassName => "Partial";
+
+		/// <inheritdoc/>
+		public override UhtClass? EngineClass => null; // Partials don't have a UClass equivalent
+
+		/// <inheritdoc/>
+		[JsonIgnore]
+		public override string EngineNamePrefix => "F";
+
+		/// <summary>
+		/// Namespace export name for partials
+		/// </summary>
+		public string NamespaceExportName => Namespace.FullSourceName;
+
+		/// <summary>
+		/// Collection of functions and other declarations found in the partial
+		/// </summary>
+		[JsonIgnore]
+		public IList<UhtDeclaration>? Declarations => _declarations;
+
+		/// <summary>
+		/// C++ code has BeginDestroy declared
+		/// </summary>
+		public bool HasBeginDestroy { get; set; }
+
+		/// <summary>
+		/// After functions are fully resolved, points to the first non-delegate function in this class
+		/// </summary>
+		public UhtDefineScopeLink<UhtFunction>? FirstFunction { get; set; }
+
+		///<inheritdoc/>
+		[JsonIgnore]
+		protected override UhtSpecifierValidatorTable? SpecifierValidatorTable => Session.GetSpecifierValidatorTable(UhtTableNames.ScriptStruct);
+
+		/// <summary>
+		/// Construct a new partial
+		/// </summary>
+		/// <param name="headerFile">Header being parsed</param>
+		/// <param name="namespaceObj">Namespace where the field was defined</param>
+		/// <param name="outer">Outer type</param>
+		/// <param name="lineNumber">Line number of the definition</param>
+		public UhtPartial(UhtHeaderFile headerFile, UhtNamespace namespaceObj, UhtType outer, int lineNumber) : base(headerFile, namespaceObj, outer, lineNumber)
+		{
+		}
+
+		/// <summary>
+		/// Construct a type from the cache
+		/// </summary>
+		/// <param name="reader">Reader</param>
+		/// <param name="outer">Outer type</param>
+		public UhtPartial(UhtInputCacheReader reader, UhtType outer) : base(reader, outer)
+		{
+			_declarations = reader.ReadOptionalList<UhtDeclaration>((reader) => new(reader));
+			OwnerClassName = reader.ReadString();
+			OwnerClass = reader.ReadType() as UhtClass;
+			MacroDeclaredLineNumber = reader.ReadInt32();
+		}
+
+		/// <summary>
+		/// Write the output type
+		/// </summary>
+		/// <param name="writer"></param>
+		public override void Write(UhtInputCacheWriter writer)
+		{
+			base.Write(writer);
+			writer.WriteOptionalVariableLengthArray(_declarations, (writer, x) => { x.Write(writer); });
+			writer.WriteString(OwnerClassName);
+			writer.WriteType(OwnerClass);
+			writer.WriteInt32(MacroDeclaredLineNumber);
+		}
+
+		#region Resolution support
+
+		/// <inheritdoc/>
+		protected override bool ResolveSelf(UhtResolvePhase resolvePhase)
+		{
+			bool result = base.ResolveSelf(resolvePhase);
+
+			switch (resolvePhase)
+			{
+				case UhtResolvePhase.Bases:
+					// Resolve the owner class reference
+					if (!String.IsNullOrEmpty(OwnerClassName))
+					{
+						OwnerClass = Session.FindType(null, UhtFindOptions.SourceName | UhtFindOptions.Class, OwnerClassName) as UhtClass;
+						if (OwnerClass == null)
+						{
+							this.LogError($"Partial '{SourceName}' extends unknown class '{OwnerClassName}'");
+						}
+						else if (!OwnerClass.ClassFlags.HasAnyFlags(EClassFlags.Native))
+						{
+							this.LogError($"Partial '{SourceName}' can only extend native classes ('{OwnerClassName}' is not native)");
+						}
+					}
+					break;
+
+				case UhtResolvePhase.Properties:
+				// Resolve property types (convert UhtPreResolveProperty to actual property types)
+				UhtPropertyParser.ResolveChildren(this, UhtPropertyParseOptions.AddModuleRelativePath);
+
+					// Validate that partial properties are compatible with the Owner class
+					if (OwnerClass != null)
+					{
+						// Partials cannot have base structs
+						if (Super != null)
+						{
+							this.LogError($"Partial '{SourceName}' cannot have a base struct (partials cannot inherit from other partials or structs)");
+						}
+
+						// Validate that all properties are compatible
+						foreach (UhtProperty property in Properties)
+						{
+							// Properties must have default constructors since they're constructed via placement new
+							if (property.PropertyCategory == UhtPropertyCategory.RegularParameter)
+							{
+								this.LogError($"Partial '{SourceName}' property '{property.SourceName}' cannot be a parameter type");
+							}
+						}
+					}
+					break;
+			}
+
+			return result;
+		}
+
+		/// <inheritdoc/>
+		protected override void ResolveChildren(UhtResolvePhase phase)
+		{
+			base.ResolveChildren(phase);
+
+			switch (phase)
+			{
+				case UhtResolvePhase.Final:
+					LinkFunctions();
+					break;
+			}
+		}
+
+		/// <summary>
+		/// Create a linked list of functions across preprocessor definition changes
+		/// </summary>
+		protected void LinkFunctions()
+		{
+			List<UhtFunction> functions = [];
+			foreach (UhtType type in Children)
+			{
+				// Delegates are not linked into class's field list
+				if (type is UhtFunction function && !function.FunctionFlags.HasAnyFlags(EFunctionFlags.Delegate))
+				{
+					functions.Add(function);
+				}
+			}
+			functions.Sort((x, y) => StringComparerUE.OrdinalIgnoreCase.Compare(x.EngineName, y.EngineName));
+
+			using UhtDefineScopeListBuilder<UhtFunction> functionList = new(DefineScope);
+			foreach (UhtFunction function in functions)
+			{
+				if (function.NextFunction is not null)
+				{
+					throw new Exception("Unexpectedly initialized next link on property");
+				}
+				function.NextFunction = functionList.Add(function);
+			}
+			FirstFunction = functionList.Head;
+		}
+
+		#endregion
+
+		#region Forward declaration support
+
+		/// <inheritdoc/>
+		public override StringBuilder AppendForwardDeclaration(StringBuilder builder)
+		{
+			return builder.Append($"struct {SourceName};");
+		}
+
+		#endregion
+
+		#region Reference collection support
+
+		/// <inheritdoc/>
+		public override void CollectReferences(IUhtReferenceCollector collector)
+		{
+			collector.AddExportType(this);
+			// Note: Partials don't add declarations or object references for themselves
+			// because they're plain structs, not UObjects with Z_Construct functions
+			if (OwnerClass != null)
+			{
+				collector.AddForwardDeclaration(OwnerClass);
+				collector.AddObjectReference(OwnerClass, UhtSingletonType.Registered);
+			}
+			foreach (UhtType child in Children)
+			{
+				child.CollectReferences(collector);
+			}
+		}
+
+		#endregion
+
+		#region Declaration support
+
+		/// <summary>
+		/// Add a declaration to the partial (for functions and other declarations)
+		/// </summary>
+		/// <param name="compilerDirectives">Currently active compiler directives</param>
+		/// <param name="tokens">List of declaration tokens</param>
+		/// <param name="function">If parsed as part of a UFUNCTION, this will reference it</param>
+		public void AddDeclaration(UhtCompilerDirective compilerDirectives, List<UhtToken> tokens, UhtFunction? function)
+		{
+			_declarations ??= [];
+			_declarations.Add(new UhtDeclaration { CompilerDirectives = compilerDirectives, Tokens = [.. tokens], Function = function });
+		}
+
+		#endregion
+
+		#region Validation support
+
+		/// <inheritdoc/>
+		protected override UhtValidationOptions Validate(UhtValidationOptions options)
+		{
+			options = base.Validate(options);
+
+			// Validate that GENERATED_BODY was found
+			if (MacroDeclaredLineNumber == -1)
+			{
+				this.LogError("Expected a GENERATED_BODY() at the start of the partial");
+			}
+
+			// Validate owner class is specified
+			if (string.IsNullOrEmpty(OwnerClassName))
+			{
+				this.LogError("Partial must specify the class it extends via UPARTIAL(ClassName) macro");
+			}
+
+			return options;
+		}
+
+		#endregion
+	}
+}
diff --git a/Engine/Source/Programs/Shared/EpicGames.UHT/Types/UhtProperty.cs b/Engine/Source/Programs/Shared/EpicGames.UHT/Types/UhtProperty.cs
--- a/Engine/Source/Programs/Shared/EpicGames.UHT/Types/UhtProperty.cs
+++ b/Engine/Source/Programs/Shared/EpicGames.UHT/Types/UhtProperty.cs
@@ -1890,6 +1890,10 @@ namespace EpicGames.UHT.Types
 				outerFunc(builder);
 				builder.Append(", ");
 			}
+			else if (property.Outer is UhtPartial partial)
+			{
+				builder.Append($"&{context.GetSingletonName(partial.OwnerClass, UhtSingletonType.ConstInit)}, ");
+			}
 			else if (property.Outer is UhtObject obj)
 			{
 				builder.Append($"&{context.GetSingletonName(obj, UhtSingletonType.ConstInit)}, ");
diff --git a/Engine/Source/Programs/Shared/EpicGames.UHT/Utils/UhtNames.cs b/Engine/Source/Programs/Shared/EpicGames.UHT/Utils/UhtNames.cs
--- a/Engine/Source/Programs/Shared/EpicGames.UHT/Utils/UhtNames.cs
+++ b/Engine/Source/Programs/Shared/EpicGames.UHT/Utils/UhtNames.cs
@@ -109,6 +109,8 @@ namespace EpicGames.UHT.Utils
 		public static readonly UhtCppIdentifier StaticsRigVMStructsId = RigVMStructsId.MakeStatics();
 		public static readonly UhtCppIdentifier VerseFuncsId = new("VerseFuncs");
 		public static readonly UhtCppIdentifier StaticsVerseFuncsId = VerseFuncsId.MakeStatics();
+		public static readonly UhtCppIdentifier PartialInfoId = new("PartialInfo");
+		public static readonly UhtCppIdentifier StaticsPartialInfoId = PartialInfoId.MakeStatics();
 
 		// Helper methods to generate cpp identifiers
 		public static UhtCppIdentifier GetPropertyIdentifier(UhtProperty? property)
diff --git a/Engine/Source/Programs/Shared/EpicGames.UHT/Utils/UhtSession.cs b/Engine/Source/Programs/Shared/EpicGames.UHT/Utils/UhtSession.cs
--- a/Engine/Source/Programs/Shared/EpicGames.UHT/Utils/UhtSession.cs
+++ b/Engine/Source/Programs/Shared/EpicGames.UHT/Utils/UhtSession.cs
@@ -2289,7 +2289,7 @@ namespace EpicGames.UHT.Utils
 				}
 
 				// Top-level objects, excluding those with alternate objects (native inteface objects)
-				Dictionary<UhtObject, int> linkableObjects = module.Packages.SelectMany(p => p.Children.OfType<UhtObject>().Where(o => o.AlternateObject is null))
+				Dictionary<UhtObject, int> linkableObjects = module.Packages.SelectMany(p => p.Children.OfType<UhtObject>().Where(o => o.AlternateObject is null && o is not UhtPartial))
 					// Objects whose outer is not a package but which need to be registered
 					.Concat(module.Headers.SelectMany(h => h.References.Singletons).Where(o => o.Package != o.Outer))
 					.OrderBy(o => (o.LinkOrder, o.SourceName, o.Outer!.SourceName))
diff --git a/Engine/Source/Programs/Shared/EpicGames.UHT/Utils/UhtTables.cs b/Engine/Source/Programs/Shared/EpicGames.UHT/Utils/UhtTables.cs
--- a/Engine/Source/Programs/Shared/EpicGames.UHT/Utils/UhtTables.cs
+++ b/Engine/Source/Programs/Shared/EpicGames.UHT/Utils/UhtTables.cs
@@ -94,6 +94,11 @@ namespace EpicGames.UHT.Utils
 		/// Table for VModules
 		/// </summary>
 		public const string VModule = "VModule";
+
+		/// <summary>
+		/// Table for UPARTIAL
+		/// </summary>
+		public const string Partial = "Partial";
 	}
 
 	/// <summary>
@@ -389,6 +394,7 @@ namespace EpicGames.UHT.Utils
 			CreateTables(tables, UhtTableNames.Function, "Function", EUhtCreateTablesFlags.Specifiers | EUhtCreateTablesFlags.Keyword, UhtTableNames.Struct, false);
 			CreateTables(tables, UhtTableNames.ScriptStruct, "Struct", EUhtCreateTablesFlags.Specifiers | EUhtCreateTablesFlags.Keyword, UhtTableNames.Struct, false);
 			CreateTables(tables, UhtTableNames.VModule, "VModule", EUhtCreateTablesFlags.Specifiers | EUhtCreateTablesFlags.Keyword, UhtTableNames.Field, false);
+			CreateTables(tables, UhtTableNames.Partial, "Partial", EUhtCreateTablesFlags.Specifiers | EUhtCreateTablesFlags.Keyword, UhtTableNames.ScriptStruct, false);
 		}
 
 		/// <summary>
diff --git a/Engine/Source/Programs/UnrealBuildTool/System/SourceFileMetadataCache.cs b/Engine/Source/Programs/UnrealBuildTool/System/SourceFileMetadataCache.cs
--- a/Engine/Source/Programs/UnrealBuildTool/System/SourceFileMetadataCache.cs
+++ b/Engine/Source/Programs/UnrealBuildTool/System/SourceFileMetadataCache.cs
@@ -134,7 +134,7 @@ namespace UnrealBuildTool
 		/// <summary>
 		/// Regex that matches C++ code with UObject declarations which we will need to generated code for.
 		/// </summary>
-		static readonly Regex ReflectionMarkupRegex = new Regex("^\\s*(UCLASS|USTRUCT|UENUM|UINTERFACE|UDELEGATE|VMODULE)\\b", RegexOptions.Compiled);
+		static readonly Regex ReflectionMarkupRegex = new Regex("^\\s*(UCLASS|USTRUCT|UENUM|UINTERFACE|UDELEGATE|UPARTIAL|VMODULE)\\b", RegexOptions.Compiled);
 
 		/// <summary>
 		/// Regex that matches #include UE_INLINE_GENERATED_CPP_BY_NAME(****) statements.
diff --git a/Engine/Source/Runtime/CoreUObject/Private/Serialization/UnversionedPropertySerialization.cpp b/Engine/Source/Runtime/CoreUObject/Private/Serialization/UnversionedPropertySerialization.cpp
--- a/Engine/Source/Runtime/CoreUObject/Private/Serialization/UnversionedPropertySerialization.cpp
+++ b/Engine/Source/Runtime/CoreUObject/Private/Serialization/UnversionedPropertySerialization.cpp
@@ -37,9 +37,27 @@ struct FDefaultStruct
 {
 	const uint8* Data;
 #if CACHE_UNVERSIONED_PROPERTY_SCHEMA
-	uint32 StructSize;
+	int32 StructStart;
+	int32 StructEnd;
 
-	FDefaultStruct(const uint8* InData, const UStruct* InStruct) : Data(InData), StructSize(InData ? InStruct->GetPropertiesSize() : 0)  {}
+	FDefaultStruct(const uint8* InData, const UStruct* InStruct)
+		: Data(InData)
+	{
+		StructStart = 0;
+		StructEnd = 0;
+		if (InData)
+		{
+			if (const UClass* StructAsClass = Cast<const UClass>(InStruct))
+			{
+				StructStart = StructAsClass->GetPropertiesStartOffset();
+				StructEnd = InStruct->GetPropertiesSize() + StructStart;
+			}
+			else
+			{
+				StructEnd = InStruct->GetPropertiesSize();
+			}
+		}
+	}
 #else
 	const UStruct* Struct;
 
@@ -97,7 +115,7 @@ public:
 	const void* GetDefaultValue(FDefaultStruct Defaults) const
 	{
 #if CACHE_UNVERSIONED_PROPERTY_SCHEMA
-		return Offset < Defaults.StructSize ? Defaults.Data + Offset : nullptr;
+		return Offset >= Defaults.StructStart && Offset < Defaults.StructEnd ? Defaults.Data + Offset : nullptr;
 #else
 		return Property->ContainerPtrToValuePtrForDefaults<uint8>(Defaults.Struct, Defaults.Data);
 #endif
@@ -323,7 +341,7 @@ private:
 
 	FProperty* Property;
 #if CACHE_UNVERSIONED_PROPERTY_SCHEMA
-	uint32 Offset;
+	int32 Offset;
 	bool bSerializeAsInteger;
 	bool bIsOptional : 1;
 #if WITH_EDITORONLY_DATA
diff --git a/Engine/Source/Runtime/CoreUObject/Private/UObject/Class.cpp b/Engine/Source/Runtime/CoreUObject/Private/UObject/Class.cpp
--- a/Engine/Source/Runtime/CoreUObject/Private/UObject/Class.cpp
+++ b/Engine/Source/Runtime/CoreUObject/Private/UObject/Class.cpp
@@ -38,6 +38,7 @@
 #include "UObject/LinkerSave.h"
 #include "UObject/InstanceDataObjectUtils.h"
 #include "UObject/Interface.h"
+#include "UObject/UObjectPartials.h"
 #include "UObject/LinkerPlaceholderClass.h"
 #include "UObject/LinkerPlaceholderFunction.h"
 #include "UObject/PropertyOptional.h"
@@ -6408,7 +6409,7 @@ PRAGMA_ENABLE_DEPRECATION_WARNINGS
 	PropertyLink = nullptr;
 	DestructorLink = nullptr;
 	CppClassStaticFunctions.Reset();
-
+	Partials = nullptr;
 	ScriptAndPropertyObjectReferences.Empty();
 	DeleteUnresolvedScriptProperties();
 
diff --git a/Engine/Source/Runtime/CoreUObject/Private/UObject/CompiledInUObjectInit.cpp b/Engine/Source/Runtime/CoreUObject/Private/UObject/CompiledInUObjectInit.cpp
--- a/Engine/Source/Runtime/CoreUObject/Private/UObject/CompiledInUObjectInit.cpp
+++ b/Engine/Source/Runtime/CoreUObject/Private/UObject/CompiledInUObjectInit.cpp
@@ -5,6 +5,7 @@
 #include "UObject/PropertyOptional.h"
 #include "UObject/UObjectBase.h"
 #include "UObject/UObjectGlobals.h"
+#include "UObject/UObjectPartials.h"
 
 #if UE_WITH_CONSTINIT_UOBJECT
 
@@ -353,6 +354,29 @@ void FCompiledInObjectRegistry::AddAndHashObjects()
 	SCOPED_BOOT_TIMING("FDeferredRegistry::AddAndHashPendingRegistrants");
 	FIterationMarker Stop = AdvanceConstruction(EConstructionStep::AddObjects);
 
+	{
+		TArray<UClass*> PartialClasses;
+		Iterate(Stop, 
+			[&](const FRegisterCompiledInObjects& Generated)
+			{
+				for (UE::CoreUObject::Private::FPartialClass* Partial : Generated.Objects.Partials)
+				{
+					UClass* PartialClass = Partial->Class;
+					if (!PartialClass->Partials)
+					{
+						PartialClasses.Add(PartialClass);
+					}
+					Partial->LinkToClass(*PartialClass);
+				}
+			}
+		);
+		Algo::Sort(PartialClasses, [](UClass* A, UClass* B) { return B->IsChildOf(A); });
+		for (UClass* PartialClass : PartialClasses)
+		{
+			UE::CoreUObject::Private::ConstructUClassPartials(*PartialClass);
+		}
+	}
+
 	// Fixup object pointers first - in order to hash a class we will need e.g. the super pointer
 	// This is only necessary for generated objects and not intrinsic ones
 #if !IS_MONOLITHIC
@@ -368,6 +392,7 @@ void FCompiledInObjectRegistry::AddAndHashObjects()
 			{
 				if (Class)
 				{ 
+					UE::CoreUObject::Private::ConstructUClassPartials(*Class);
 					Class->LinkCompiledInPointerFields(Linker);
 					for (UField* Field = Class->Children; Field; Field = Field->Next)
 					{
@@ -420,10 +445,15 @@ void FCompiledInObjectRegistry::AddAndHashObjects()
 				}
 			}
 		}
+	});
 
+	IterateObjects(Stop, [](UObject* Obj)
+	{
 		// TODO: General purpose interface function for initializing compiled in FNames
 		if (UClass* Class = Cast<UClass>(Obj))
 		{
+			UE::CoreUObject::Private::ConstructUClassPartials(*Class); // We need to call this again because there might be subclasses to partial classes that has no partials
+
 			Class->ClassConfigName = FName(Class->CompiledInClassConfigName);
 
 			// Create TArray of native functions from compiled-in data to maintain API
@@ -432,6 +462,16 @@ void FCompiledInObjectRegistry::AddAndHashObjects()
 			{
 				TConstArrayView<UE::CodeGen::FClassNativeFunction> CompiledInNativeFunctions = Class->CompiledInNativeFunctions;
 				new (&Class->NativeFunctionLookupTable) TArray<FNativeFunctionLookup>(CompiledInNativeFunctions);
+
+				// Partials have functions too that needs to be added to the NativeFunctionLookupTable
+				for (UE::CoreUObject::Private::FPartialClass* It = Class->Partials; It ; It = It->NextPartial)
+				{
+					for (UField* Child = It->FirstChild; Child != It->EndChild; Child = Child->Next)
+					{
+						UFunction* Func = (UFunction*)Child;
+						Class->AddNativeFunction(*Func->GetName(), Func->GetNativeFunc());
+					}
+				}
 			}
 		}
 		else if (USparseDelegateFunction* SparseDelegate = Cast<USparseDelegateFunction>(Obj))
diff --git a/Engine/Source/Runtime/CoreUObject/Private/UObject/DeferredRegistry.h b/Engine/Source/Runtime/CoreUObject/Private/UObject/DeferredRegistry.h
--- a/Engine/Source/Runtime/CoreUObject/Private/UObject/DeferredRegistry.h
+++ b/Engine/Source/Runtime/CoreUObject/Private/UObject/DeferredRegistry.h
@@ -32,6 +32,11 @@ inline bool operator != (const FPackageReloadVersionInfo& lhs, const FPackageRel
 	return lhs.BodyHash != rhs.BodyHash || lhs.DeclarationsHash != rhs.DeclarationsHash;
 }
 
+inline bool operator != (const FPartialReloadVersionInfo& lhs, const FPartialReloadVersionInfo& rhs)
+{
+	return lhs.Size != rhs.Size || lhs.Hash != rhs.Hash;
+}
+
 void ReloadProcessObject(UScriptStruct* Object, const TCHAR* RenamePrefix);
 void ReloadProcessObject(UEnum* Object, const TCHAR* RenamePrefix);
 void ReloadProcessObject(UClass* Object, const TCHAR* RenamePrefix);
@@ -364,5 +369,6 @@ using FClassDeferredRegistry = TDeferredRegistry<FClassRegistrationInfo>;
 using FEnumDeferredRegistry = TDeferredRegistry<FEnumRegistrationInfo>;
 using FStructDeferredRegistry = TDeferredRegistry<FStructRegistrationInfo>;
 using FPackageDeferredRegistry = TDeferredRegistry<FPackageRegistrationInfo>;
+using FPartialDeferredRegistry = TDeferredRegistry<FPartialRegistrationInfo>;
 
 #endif // !UE_WITH_CONSTINIT_UOBJECT
diff --git a/Engine/Source/Runtime/CoreUObject/Private/UObject/GarbageCollection.cpp b/Engine/Source/Runtime/CoreUObject/Private/UObject/GarbageCollection.cpp
--- a/Engine/Source/Runtime/CoreUObject/Private/UObject/GarbageCollection.cpp
+++ b/Engine/Source/Runtime/CoreUObject/Private/UObject/GarbageCollection.cpp
@@ -26,6 +26,7 @@
 #include "UObject/ObjectPtr.h"
 #include "UObject/Class.h"
 #include "UObject/UObjectIterator.h"
+#include "UObject/UObjectPartials.h"
 #include "UObject/UnrealType.h"
 #include "UObject/LinkerLoad.h"
 #include "UObject/GCObject.h"
@@ -950,8 +951,17 @@ public:
 				if (Object)
 #endif // UE_WITH_CONSTINIT_UOBJECT
 				{
+					int32 OffsetToObject = 0;
+
+					if (Object->HasAnyFlags(RF_HasPartials))
+					{
+						UClass* Class = Object->GetClass();
+						UE::CoreUObject::Private::DestructPartials(*Class, Object);
+						OffsetToObject = -Class->GetPropertiesStartOffset();
+					}
+
 					Object->~UObject();
-					GUObjectAllocator.FreeUObject(Object);
+					GUObjectAllocator.FreeUObject(Object, OffsetToObject);
 					UnreachableObject.Object = nullptr;
 				}
 
@@ -3947,7 +3957,7 @@ void TFastReferenceCollector<ProcessorType, CollectorType>::ProcessStructs(Dispa
 				typename DispatcherType::SchemaStackScopeType SchemaStack(Dispatcher.Context, AoS.Schema);
 				do
 				{
-					Private::VisitMembers(Dispatcher, AoS.Schema, StructIt); 
+					Private::VisitMembers(Dispatcher, AoS.Schema, StructIt, 0); 
 					StructIt += AoS.Stride;
 				}
 				while (StructIt != StructEnd);
@@ -6999,13 +7009,16 @@ void UClass::AssembleReferenceTokenStreamInternal(bool bForce)
 			ClassFlags &= ~CLASS_TokenStreamAssembled;
 		}
 
+		// We need to make sure all offsets to properties are positive because of schema format.
+		int32 StartOffset = -GetPropertiesStartOffset();
+
 		FSchemaBuilder Schema(/* don't store sizeof(class) to enable super class schema reuse */ 0);
 		FSchemaView SuperSchema;
 		if (UClass* SuperClass = GetSuperClass())
 		{
 			SuperClass->AssembleReferenceTokenStreamInternal();	
 			SuperSchema = SuperClass->ReferenceSchema.Get();
-			Schema.Append(SuperSchema);
+			Schema.Append(SuperSchema, StartOffset - (-SuperClass->GetPropertiesStartOffset()));
 		}
 		const int32 NumSuperMembers = Schema.NumMembers();
 
@@ -7018,13 +7031,13 @@ void UClass::AssembleReferenceTokenStreamInternal(bool bForce)
 			{
 				FProperty* Property = *It;
 				FPropertyStackScope PropertyScope(DebugPath, Property);
-				Property->EmitReferenceInfo(Schema, 0, EncounteredStructProps, DebugPath);
+				Property->EmitReferenceInfo(Schema, StartOffset, EncounteredStructProps, DebugPath);
 			}
 		}
 		
 		if (ClassFlags & CLASS_Intrinsic)
 		{
-			Schema.Append(UE::GC::GetIntrinsicSchema(this));
+			Schema.Append(UE::GC::GetIntrinsicSchema(this), 0);
 		}
 		
 		// Make sure all Blueprint properties are marked as non-native
diff --git a/Engine/Source/Runtime/CoreUObject/Private/UObject/GarbageCollectionSchema.cpp b/Engine/Source/Runtime/CoreUObject/Private/UObject/GarbageCollectionSchema.cpp
--- a/Engine/Source/Runtime/CoreUObject/Private/UObject/GarbageCollectionSchema.cpp
+++ b/Engine/Source/Runtime/CoreUObject/Private/UObject/GarbageCollectionSchema.cpp
@@ -328,7 +328,7 @@ static void CopyDebugNames(const FName* In, TArrayView<FMemberDeclaration> Out)
 	}
 }
 
-void FSchemaBuilder::Append(FSchemaView SuperSchema)
+void FSchemaBuilder::Append(FSchemaView SuperSchema, int32 OffsetForProperties)
 {
 	if (SuperSchema.IsEmpty())
 	{
@@ -356,7 +356,7 @@ void FSchemaBuilder::Append(FSchemaView SuperSchema)
 			}
 			else
 			{
-				Members.Add({ FName(), (JumpedWords + Member.WordOffset) * 8, Member.Type, ExtraWord});
+				Members.Add({ FName(), OffsetForProperties + (JumpedWords + Member.WordOffset) * 8, Member.Type, ExtraWord});
 				TryAddSchemaReference(Members.Last());
 			}
 		}
diff --git a/Engine/Source/Runtime/CoreUObject/Private/UObject/Obj.cpp b/Engine/Source/Runtime/CoreUObject/Private/UObject/Obj.cpp
--- a/Engine/Source/Runtime/CoreUObject/Private/UObject/Obj.cpp
+++ b/Engine/Source/Runtime/CoreUObject/Private/UObject/Obj.cpp
@@ -28,6 +28,7 @@
 #include "Serialization/ArchiveUObject.h"
 #include "UObject/GarbageCollectionHistory.h"
 #include "UObject/Class.h"
+#include "UObject/UObjectPartials.h"
 #include "UObject/EnumProperty.h"
 #include "UObject/UObjectIterator.h"
 #include "UObject/Package.h"
@@ -1109,6 +1110,8 @@ void UObject::BeginDestroy()
 			);
 	}
 
+	UE::CoreUObject::Private::BeginDestroyPartials(*GetClass(), this);
+
 	// Remove from linker's export table.
 	SetLinker( NULL, INDEX_NONE );
 
diff --git a/Engine/Source/Runtime/CoreUObject/Private/UObject/UObjectAllocator.cpp b/Engine/Source/Runtime/CoreUObject/Private/UObject/UObjectAllocator.cpp
--- a/Engine/Source/Runtime/CoreUObject/Private/UObject/UObjectAllocator.cpp
+++ b/Engine/Source/Runtime/CoreUObject/Private/UObject/UObjectAllocator.cpp
@@ -31,7 +31,7 @@ void FUObjectAllocator::DisablePersistentAllocator()
  * @return newly allocated UObjectBase (not really a UObjectBase yet, no constructor-like thing has been called).
  */
 UE_AUTORTFM_ALWAYS_OPEN
-UObjectBase* FUObjectAllocator::AllocateUObject(int32 Size, int32 Alignment, bool bAllowPermanent)
+UObjectBase* FUObjectAllocator::AllocateUObject(int32 Size, int32 Alignment, bool bAllowPermanent, int32 OffsetToObject)
 {
 	void* Result = nullptr;
 
@@ -46,7 +46,7 @@ UObjectBase* FUObjectAllocator::AllocateUObject(int32 Size, int32 Alignment, boo
 		Result = FMemory::Malloc(Size, Alignment);
 	}
 
-	return (UObjectBase*)Result;
+	return (UObjectBase*)(((uint8*)Result) + OffsetToObject);
 }
 
 /**
@@ -54,13 +54,13 @@ UObjectBase* FUObjectAllocator::AllocateUObject(int32 Size, int32 Alignment, boo
  *
  * @param Object object to free
  */
-void FUObjectAllocator::FreeUObject(UObjectBase *Object) const
+void FUObjectAllocator::FreeUObject(UObjectBase *Object, int32 OffsetToObject) const
 {
 	check(Object);
 	// Only free memory if it was allocated directly from allocator and not from permanent object pool.
 	if (FPermanentObjectPoolExtents().Contains(Object) == false)
 	{
-		FMemory::Free(Object);
+		FMemory::Free(((uint8*)Object) - OffsetToObject);
 	}
 	// We only destroy objects residing in permanent object pool during the exit purge.
 	else
diff --git a/Engine/Source/Runtime/CoreUObject/Private/UObject/UObjectBase.cpp b/Engine/Source/Runtime/CoreUObject/Private/UObject/UObjectBase.cpp
--- a/Engine/Source/Runtime/CoreUObject/Private/UObject/UObjectBase.cpp
+++ b/Engine/Source/Runtime/CoreUObject/Private/UObject/UObjectBase.cpp
@@ -20,6 +20,7 @@
 #include "UObject/DeferredRegistry.h"
 #include "UObject/ObjectHandlePrivate.h"
 #include "UObject/UObjectIterator.h"
+#include "UObject/UObjectPartials.h"
 #include "UObject/UObjectStats.h"
 #include "UObject/UObjectPrivate.h"
 #include "UObject/Package.h"
@@ -783,7 +784,7 @@ void RegisterCompiledInInfo(UPackage* (*InOuterRegister)(ETypeConstructPhase), c
 }
 
 // Multiple registrations
-void RegisterCompiledInInfo(const TCHAR* PackageName, const FClassRegisterCompiledInInfo* ClassInfo, size_t NumClassInfo, const FStructRegisterCompiledInInfo* StructInfo, size_t NumStructInfo, const FEnumRegisterCompiledInInfo* EnumInfo, size_t NumEnumInfo)
+void RegisterCompiledInInfo(const TCHAR* PackageName, const FClassRegisterCompiledInInfo* ClassInfo, size_t NumClassInfo, const FStructRegisterCompiledInInfo* StructInfo, size_t NumStructInfo, const FEnumRegisterCompiledInInfo* EnumInfo, size_t NumEnumInfo, const FPartialRegisterCompiledInInfo *PartialsInfo, size_t NumPartialsInfo)
 {
 	LLM_SCOPE(ELLMTag::UObject);
 	LLM_SCOPE_BYTAG(UObject_UObjectBase);
@@ -809,6 +810,12 @@ void RegisterCompiledInInfo(const TCHAR* PackageName, const FClassRegisterCompil
 		const FEnumRegisterCompiledInInfo& Info = EnumInfo[Index];
 		RegisterCompiledInInfo(Info.OuterRegister, PackageName, Info.Name, *Info.Info, Info.VersionInfo);
 	}
+
+	for (size_t Index = 0; Index < NumPartialsInfo; ++Index)
+	{
+		const FPartialRegisterCompiledInInfo& Info = PartialsInfo[Index];
+		FPartialDeferredRegistry::Get().AddRegistration(Info.OuterRegister, PackageName, Info.Name, *Info.Info, Info.VersionInfo);
+	}
 }
 
 /** Register all loaded classes */
@@ -1090,6 +1097,7 @@ void ProcessNewlyLoadedUObjects(FName InModuleName, bool bCanProcessNewlyLoadedO
 	FClassDeferredRegistry& ClassRegistry = FClassDeferredRegistry::Get();
 	FStructDeferredRegistry& StructRegistry = FStructDeferredRegistry::Get();
 	FEnumDeferredRegistry& EnumRegistry = FEnumDeferredRegistry::Get();
+	FPartialDeferredRegistry& PartialRegistry = FPartialDeferredRegistry::Get();
 
 	PackageRegistry.ProcessChangedObjects(true);
 	StructRegistry.ProcessChangedObjects();
@@ -1102,7 +1110,8 @@ void ProcessNewlyLoadedUObjects(FName InModuleName, bool bCanProcessNewlyLoadedO
 	while (GFirstPendingRegistrant ||
 		ClassRegistry.HasPendingRegistrations() ||
 		StructRegistry.HasPendingRegistrations() ||
-		EnumRegistry.HasPendingRegistrations())
+		EnumRegistry.HasPendingRegistrations() ||
+		PartialRegistry.HasPendingRegistrations())
 	{
 		if (bNewUObjects)
 		{
@@ -1114,6 +1123,16 @@ void ProcessNewlyLoadedUObjects(FName InModuleName, bool bCanProcessNewlyLoadedO
 		}
 
 		bNewUObjects = true;
+
+		// Initializing partials first since they need to be linked to class before outer construction
+		for (const FPartialDeferredRegistry::FRegistrant& Registrant : PartialRegistry.GetRegistrations())
+		{
+			UE::CoreUObject::Private::FPartialClass* Params = Registrant.RegisterFn(ETypeConstructPhase::Outer);
+			Params->LinkToClass(*Params->Class(ETypeConstructPhase::Inner));
+		}
+		PartialRegistry.EmptyRegistrations();
+
+
 		UObjectProcessRegistrants();
 		UObjectLoadAllCompiledInStructs();
 
diff --git a/Engine/Source/Runtime/CoreUObject/Private/UObject/UObjectConstructInternal.h b/Engine/Source/Runtime/CoreUObject/Private/UObject/UObjectConstructInternal.h
--- a/Engine/Source/Runtime/CoreUObject/Private/UObject/UObjectConstructInternal.h
+++ b/Engine/Source/Runtime/CoreUObject/Private/UObject/UObjectConstructInternal.h
@@ -246,6 +246,8 @@ void ConstructUClassHelper(UClass*& OutClass, const ClassParams& Params, PostNew
 
 	ConstructFProperties(NewClass, Params.PropertyArray, Params.NumProperties);
 
+	UE::CoreUObject::Private::ConstructUClassPartials(*NewClass);
+
 	if (Params.ClassConfigNameUTF8)
 	{
 		NewClass->ClassConfigName = FName(UTF8_TO_TCHAR(Params.ClassConfigNameUTF8));
diff --git a/Engine/Source/Runtime/CoreUObject/Private/UObject/UObjectGlobals.cpp b/Engine/Source/Runtime/CoreUObject/Private/UObject/UObjectGlobals.cpp
--- a/Engine/Source/Runtime/CoreUObject/Private/UObject/UObjectGlobals.cpp
+++ b/Engine/Source/Runtime/CoreUObject/Private/UObject/UObjectGlobals.cpp
@@ -33,6 +33,7 @@
 #include "UObject/GarbageCollection.h"
 #include "UObject/GCScopeLock.h"
 #include "UObject/Class.h"
+#include "UObject/UObjectPartials.h"
 #include "UObject/CoreRedirects.h"
 #include "UObject/FastReferenceCollector.h"
 #include "UObject/ObjectInstancingGraph.h"
@@ -3466,6 +3467,13 @@ UObject* StaticAllocateObject
 	int32 TotalSize = InClass->GetPropertiesSize();
 	checkSlow(TotalSize);
 
+	int32 OffsetToObject = -InClass->GetPropertiesStartOffset(); // Start offset for properties can be negative so we want to move object forward in allocation
+
+	if (InClass->HasAnyClassFlags(CLASS_Partial))
+	{
+		InFlags |= RF_HasPartials;
+	}
+
 	int32 OldIndex = -1;
 	int32 OldSerialNumber = 0;
 #if UE_WITH_REMOTE_OBJECT_HANDLE
@@ -3477,7 +3485,7 @@ UObject* StaticAllocateObject
 	if( Obj == nullptr )
 	{	
 		int32 Alignment	= FMath::Max( 4, InClass->GetMinAlignment() );
-		Obj = (UObject *)GUObjectAllocator.AllocateUObject(TotalSize,Alignment,GIsInitialLoad);
+		Obj = (UObject *)GUObjectAllocator.AllocateUObject(TotalSize, Alignment, GIsInitialLoad, OffsetToObject);
 	}
 	else
 	{
@@ -3668,7 +3676,7 @@ UObject* StaticAllocateObject
 
 		UE_AUTORTFM_OPEN_NO_VALIDATION
 		{
-			FMemory::Memzero((void *)Obj, TotalSize);
+			FMemory::Memzero((uint8 *)Obj - OffsetToObject, TotalSize);
 			new ((void *)Obj) UObjectBase(const_cast<UClass*>(InClass), InFlags|RF_NeedInitialization, InternalSetFlags, InOuter, InName, OldIndex, OldSerialNumber == 0 ? SerialNumber : OldSerialNumber, 
 #if UE_WITH_REMOTE_OBJECT_HANDLE
 				OldRemoteId.IsValid() ? OldRemoteId :
@@ -4133,6 +4141,8 @@ void FObjectInitializer::PostConstructInit()
 	}
 #endif
 
+	UE::CoreUObject::Private::ConstructPartials(*Class, Obj);
+
 	if (bShouldInitializePropsFromArchetype)
 	{
 		UClass* BaseClass = (bIsCDO && !GIsDuplicatingClassForReinstancing) ? SuperClass : Class;
@@ -6087,7 +6097,7 @@ namespace UECodeGen_Private
 		return NewProp;
 	}
 
-	void ConstructFProperty(FFieldVariant Outer, const FPropertyParamsBase* const*& PropertyArray, int32& NumProperties)
+	FProperty* ConstructFProperty(FFieldVariant Outer, const FPropertyParamsBase* const*& PropertyArray, int32& NumProperties)
 	{
 		const FPropertyParamsBase* PropBase = *--PropertyArray;
 
@@ -6382,6 +6392,7 @@ namespace UECodeGen_Private
 		{
 			ConstructFProperty(NewProp, PropertyArray, NumProperties);
 		}
+		return NewProp;
 	}
 
 	void ConstructFProperties(UObject* Outer, const FPropertyParamsBase* const* PropertyArray, int32 NumProperties)
diff --git a/Engine/Source/Runtime/CoreUObject/Private/UObject/UObjectPartials.cpp b/Engine/Source/Runtime/CoreUObject/Private/UObject/UObjectPartials.cpp
new file mode 100644
--- /dev/null
+++ b/Engine/Source/Runtime/CoreUObject/Private/UObject/UObjectPartials.cpp
@@ -0,0 +1,228 @@
+// Copyright Epic Games, Inc. All Rights Reserved.
+
+#include "UObject/UObjectPartials.h"
+#include "UObject/Class.h"
+#include "UObject/UnrealType.h"
+
+////////////////////////////////////////////////////////////////////////////////////////////////////
+
+namespace UECodeGen_Private
+{
+	FProperty* ConstructFProperty(FFieldVariant Outer, const FPropertyParamsBase* const*& PropertyArray, int32& NumProperties);
+}
+
+namespace UE::CoreUObject::Private
+{
+	void ConstructUClassPartials(UClass& Class)
+	{
+		if (!Class.HasAnyClassFlags(CLASS_Partial) || Class.PropertiesStartOffset < 0)
+		{
+			return;
+		}
+
+		UClass* SuperClass = Class.GetSuperClass();
+		check(SuperClass); // No root class will have CLASS_Partial set
+		ConstructUClassPartials(*SuperClass);
+
+		FPartialClass* FirstPartial = Class.Partials;
+		if (!FirstPartial)
+		{
+			if (SuperClass->PropertiesStartOffset < 0)
+			{
+				Class.MinAlignment = FMath::Max(Class.MinAlignment, SuperClass->MinAlignment); // Update alignment in case there are partials that increased alignment in base class
+				Class.PropertiesStartOffset = SuperClass->PropertiesStartOffset;
+				Class.SetPropertiesSize(Class.GetPropertiesSize() - SuperClass->PropertiesStartOffset);
+			}
+			return;
+		}
+
+		check(Class.HasAnyClassFlags(CLASS_Partial));
+
+		int32 PropertiesStartOffset = SuperClass->PropertiesStartOffset;
+		int16 MemoryMinAlignment = FMath::Max(Class.MinAlignment, SuperClass->MinAlignment);
+
+		// Calculated needed size based on partials size and alignments
+		int32 NeededSize = 0;
+		for (FPartialClass* It = FirstPartial; It ; It = It->NextPartial)
+		{
+#if UE_WITH_CONSTINIT_UOBJECT
+			if (It->FirstProperty)
+			{
+				NeededSize = Align(NeededSize, It->Alignment);
+				NeededSize += It->Size;
+			}
+#else
+			if (It->NumProperties)
+			{
+				NeededSize = Align(NeededSize, It->Alignment);
+				NeededSize += It->Size;
+			}
+#endif
+		}
+		// PropertiesStartOffset must take alignment of data stored after current partials in to account
+		PropertiesStartOffset = -Align(-PropertiesStartOffset + NeededSize, MemoryMinAlignment);
+
+
+
+		// Create the properties/functions of partials
+
+		int32 PartialOffset = PropertiesStartOffset;
+
+		for (FPartialClass* It = FirstPartial; It ; It = It->NextPartial)
+		{
+			PartialOffset = Align(PartialOffset, It->Alignment);
+			It->PartialOffset = PartialOffset;
+			*It->PartialOffsetPtr = PartialOffset;
+
+#if UE_WITH_CONSTINIT_UOBJECT
+			if (It->FirstProperty)
+			{
+				FProperty* LastProperty = nullptr;
+				for (FProperty* P = It->FirstProperty; P; P = (FProperty*)(P->Next))
+				{
+					UEProperty_Private::FProperty_DoNotUse::Unsafe_AlterOffset(*P, PartialOffset + P->GetOffset_ForInternal());
+					LastProperty = P;
+				}
+				LastProperty->Next = Class.ChildProperties;
+				Class.ChildProperties = It->FirstProperty;
+				PartialOffset += It->Size;
+			}
+#else
+			if (int32 NumProperties = It->NumProperties)
+			{
+				const UECodeGen_Private::FPropertyParamsBase* const* PropertyArray = It->Properties + NumProperties;
+				while (PropertyArray > It->Properties)
+				{
+					FProperty* Property = UECodeGen_Private::ConstructFProperty(&Class, PropertyArray, NumProperties);
+					UEProperty_Private::FProperty_DoNotUse::Unsafe_AlterOffset(*Property, PartialOffset + Property->GetOffset_ForInternal());
+				}
+				PartialOffset += It->Size;
+			}
+#endif
+
+#if UE_WITH_CONSTINIT_UOBJECT
+			if (It->FirstChild)
+			{
+				UField* LastChild = nullptr;
+				for (UField* Child = It->FirstChild; Child; Child = Child->Next)
+				{
+					LastChild = Child;
+				}
+				It->EndChild = Class.Children;
+				LastChild->Next = Class.Children;
+				Class.Children = It->FirstChild;
+			}
+#else
+			if (It->NumFunctions > 0)
+			{
+				const UE::CodeGen::FClassNativeFunction* NativeFuncs = static_cast<const UE::CodeGen::FClassNativeFunction*>(It->NativeFunctions);
+				FNativeFunctionRegistrar::RegisterFunctions(&Class, MakeConstArrayView(NativeFuncs, It->NumFunctions));
+
+				// Hook up functions 
+				for (int32 FuncIdx = 0; FuncIdx < It->NumFunctions; ++FuncIdx)
+				{
+					UFunction* (*FuncConstructor)(ETypeConstructPhase) = It->FunctionConstructors[FuncIdx];
+					check(FuncConstructor);
+					UFunction* Func = FuncConstructor(ETypeConstructPhase::Outer);
+					check(Func);
+					Func->Next = Class.Children;
+					Class.Children = Func;
+					Class.AddFunctionToFunctionMap(Func, Func->GetFName());
+				}
+			}
+#endif
+		}
+
+		check(((-PropertiesStartOffset) % Class.MinAlignment) == 0);
+
+		Class.SetPropertiesSize(Class.GetPropertiesSize() - PropertiesStartOffset);
+		Class.PropertiesStartOffset = PropertiesStartOffset;
+		Class.MinAlignment = MemoryMinAlignment;
+	}
+
+	////////////////////////////////////////////////////////////////////////////////////////////////////
+
+	void FPartialClass::LinkToClass(UClass& PartialClass)
+	{
+		// Sort in partials based on alignment first and then name
+		FPartialClass* PrevPartial = nullptr;
+		for (FPartialClass* It = PartialClass.Partials; It; It = It->NextPartial)
+		{
+			if (Alignment < It->Alignment || (Alignment == It->Alignment && FCString::Strcmp(Name, It->Name) < 0))
+			{
+				break;
+			}
+			PrevPartial = It;
+		}
+		if (!PrevPartial)
+		{
+			NextPartial = PartialClass.Partials;
+			PartialClass.Partials = this;
+		}
+		else
+		{
+			NextPartial = PrevPartial->NextPartial;
+			PrevPartial->NextPartial = this;
+		}
+	}
+
+	////////////////////////////////////////////////////////////////////////////////////////////////////
+
+	void ConstructPartials(const UClass& Class, UObject* Obj)
+	{
+		if (Class.HasAnyClassFlags(CLASS_Partial))
+		{
+			ConstructPartials(*Class.GetSuperClass(), Obj);
+			for (FPartialClass* It = Class.Partials; It; It = It->NextPartial)
+			{
+				check(It->PartialOffset != MAX_int32);
+				It->Constructor((uint8*)Obj + It->PartialOffset); // Call partial constructor to initialize C++ objects
+			}
+		}
+	}
+
+	void BeginDestroyPartials(const UClass& Class, UObject* Obj)
+	{
+		if (Class.HasAnyClassFlags(CLASS_Partial))
+		{
+			for (FPartialClass* It = Class.Partials; It; It = It->NextPartial)
+			{
+				if (It->BeginDestroy)
+				{
+					It->BeginDestroy((uint8*)Obj + It->PartialOffset);
+				}
+			}
+			BeginDestroyPartials(*Class.GetSuperClass(), Obj);
+		}
+	}
+
+	void DestructPartials(const UClass& Class, UObject* Obj)
+	{
+		if (Class.HasAnyClassFlags(CLASS_Partial))
+		{
+			for (FPartialClass* It = Class.Partials; It; It = It->NextPartial)
+			{
+				It->Destructor((uint8*)Obj + It->PartialOffset);
+			}
+			DestructPartials(*Class.GetSuperClass(), Obj);
+		}
+	}
+
+	void CallPartialsGetLifetimeReplicatedProps(const UClass& Class, UObject* Obj, TArray<FLifetimeProperty>& OutLifetimeProps)
+	{
+		if (Class.HasAnyClassFlags(CLASS_Partial))
+		{
+			CallPartialsGetLifetimeReplicatedProps(*Class.GetSuperClass(), Obj, OutLifetimeProps);
+			for (FPartialClass* It = Class.Partials; It; It = It->NextPartial)
+			{
+				if (It->BeginDestroy)
+				{
+					It->GetLifetimeReplicatedProps((uint8*)Obj + It->PartialOffset, OutLifetimeProps);
+				}
+			}
+		}
+	}
+
+} // UE::CoreUObject::Private
+
+////////////////////////////////////////////////////////////////////////////////////////////////////
diff --git a/Engine/Source/Runtime/CoreUObject/Public/UObject/Class.h b/Engine/Source/Runtime/CoreUObject/Public/UObject/Class.h
--- a/Engine/Source/Runtime/CoreUObject/Public/UObject/Class.h
+++ b/Engine/Source/Runtime/CoreUObject/Public/UObject/Class.h
@@ -132,6 +132,8 @@ struct FUObjectSerializeContext;
 template <typename FuncType> class TFunctionRef;
 enum class EPropertyObjectReferenceType : uint32;
 
+namespace UE::CoreUObject::Private { void ConstructUClassPartials(UClass&); }
+
 COREUOBJECT_API DECLARE_LOG_CATEGORY_EXTERN(LogClass, Log, All);
 COREUOBJECT_API DECLARE_LOG_CATEGORY_EXTERN(LogScriptSerialization, Log, All);
 
@@ -3889,7 +3891,12 @@ protected:
 
 	/** TRUE if instances of this class contain non-serialized references that require construction on load via instancing. */
 	bool bNeedsPostLoadSubobjectInstancing = false;
-	
+
+private:
+	/** When using partials they are allocated in front of UObject, and then this value will be negative. Note that PropertiesSize includes the negative part too */
+	int32 PropertiesStartOffset = 0;
+
+	friend void UE::CoreUObject::Private::ConstructUClassPartials(UClass&);
 public:
 
 	/** Class flags; See EClassFlags for more information */
@@ -4113,6 +4120,9 @@ public:
 #endif // UE_WITH_CONSTINIT_UOBJECT
 	};
 
+	/** Linked list of partials extending this class */
+	UE::CoreUObject::Private::FPartialClass* Partials = nullptr;
+
 public:
 	// Constructors
 	COREUOBJECT_API UClass(const FObjectInitializer& ObjectInitializer = FObjectInitializer::Get());
@@ -4304,6 +4314,12 @@ public:
 	UE_INTERNAL COREUOBJECT_API virtual void SerializePropertyValueFlags(EPropertyValueFlags Flags, void* Data, FStructuredArchiveRecord Record, FArchiveFieldName Name) const override;
 #endif // WITH_EDITORONLY_DATA
 
+	/** Offset from UObject pointer to where properties start. This method will return a negative value if partials are present */
+	UE_FORCEINLINE_HINT int32 GetPropertiesStartOffset() const
+	{
+		return PropertiesStartOffset;
+	}
+
 	/**
 	 * Returns class path name as a package + class FName pair
 	 */
diff --git a/Engine/Source/Runtime/CoreUObject/Public/UObject/FastReferenceCollector.h b/Engine/Source/Runtime/CoreUObject/Public/UObject/FastReferenceCollector.h
--- a/Engine/Source/Runtime/CoreUObject/Public/UObject/FastReferenceCollector.h
+++ b/Engine/Source/Runtime/CoreUObject/Public/UObject/FastReferenceCollector.h
@@ -684,12 +684,13 @@ FORCEINLINE EOrigin GetSchemaOrigin(FSchemaView Schema, UObject* Instance)
 #endif // UE_WITH_REMOTE_OBJECT_HANDLE
 
 template<class DispatcherType, typename ObjectType>
-AUTORTFM_INFER FORCEINLINE_DEBUGGABLE void VisitMembers(DispatcherType& Dispatcher, FSchemaView Schema, ObjectType* Instance)
+AUTORTFM_INFER FORCEINLINE_DEBUGGABLE void VisitMembers(DispatcherType& Dispatcher, FSchemaView Schema, ObjectType* Instance, int32 OffsetToPropertiesStart)
 {
 	check(!Schema.IsEmpty());
 
 	const EOrigin Origin = GetSchemaOrigin(Schema, Instance);
-	uint64* InstanceCursor = (uint64*)Instance;	// Advanced via Jump to reach far members
+	uint8* InstanceBegin = (uint8*)Instance + OffsetToPropertiesStart;
+	uint64* InstanceCursor = (uint64*)InstanceBegin;	// Advanced via Jump to reach far members
 	uint32 DebugIdx = 0;
 	for (const FMemberWord* WordIt = Schema.GetWords(); true; ++WordIt)
 	{
@@ -753,7 +754,7 @@ template<class DispatcherType>
 void VisitNestedStructMembers(DispatcherType& Dispatcher, FSchemaView Schema, uint8* Instance)
 {
 	static_assert(!DispatcherType::bBatching);
-	VisitMembers(Dispatcher, Schema, Instance);
+	VisitMembers(Dispatcher, Schema, Instance, 0);
 }
 
 } // namespace Private
@@ -1026,9 +1027,10 @@ private:
 
 			if (!Schema.IsEmpty())
 			{
+				int32 OffsetToProperty = Class->GetPropertiesStartOffset();
 				typename DispatcherType::SchemaStackScopeType SchemaStack(Dispatcher.Context, Schema);
 				Processor.BeginTimingObject(CurrentObject);
-				Private::VisitMembers(Dispatcher, Schema, CurrentObject);
+				Private::VisitMembers(Dispatcher, Schema, CurrentObject, OffsetToProperty);
 				Processor.UpdateDetailedStats(CurrentObject);
 			}
 		}
@@ -1045,7 +1047,7 @@ private:
 			FSchemaView Schema = VerseStruct->ReferenceSchema.Get();
 			check(!Schema.IsEmpty()); // Native structs with empty schemas should not be in the list to begin with
 			typename DispatcherType::SchemaStackScopeType SchemaStack(Dispatcher.Context, Schema);
-			Private::VisitMembers(Dispatcher, Schema, (uint8*)NativeStruct->GetStruct());
+			Private::VisitMembers(Dispatcher, Schema, (uint8*)NativeStruct->GetStruct(), 0);
 		}
 	}
 #endif
diff --git a/Engine/Source/Runtime/CoreUObject/Public/UObject/GarbageCollectionSchema.h b/Engine/Source/Runtime/CoreUObject/Public/UObject/GarbageCollectionSchema.h
--- a/Engine/Source/Runtime/CoreUObject/Public/UObject/GarbageCollectionSchema.h
+++ b/Engine/Source/Runtime/CoreUObject/Public/UObject/GarbageCollectionSchema.h
@@ -270,7 +270,7 @@ public:
 	COREUOBJECT_API void Add(FMemberDeclaration Member);
 
 	// Append all but the last Stop/ARO member
-	COREUOBJECT_API void Append(FSchemaView SuperSchema);
+	COREUOBJECT_API void Append(FSchemaView SuperSchema, int32 OffsetForProperties = 0);
 
 	// Emit Stop or ARO and build a view that builder holds reference to. Returns same view if called multiple times.
 	COREUOBJECT_API FSchemaView Build(ObjectAROFn ARO = nullptr);
diff --git a/Engine/Source/Runtime/CoreUObject/Public/UObject/ObjectMacros.h b/Engine/Source/Runtime/CoreUObject/Public/UObject/ObjectMacros.h
--- a/Engine/Source/Runtime/CoreUObject/Public/UObject/ObjectMacros.h
+++ b/Engine/Source/Runtime/CoreUObject/Public/UObject/ObjectMacros.h
@@ -221,8 +221,8 @@ enum EClassFlags
 	CLASS_ProjectUserConfig	  = 0x00000040u,
 	/** Class is a native class - native interfaces will have CLASS_Native set, but not RF_MarkAsNative */
 	CLASS_Native			  = 0x00000080u,
-
-	// CLASS_Unused			  = 0x00000100u,
+	/** Class can contain partials */
+	CLASS_Partial			  = 0x00000100u,
 
 	/** Do not allow users to create in the editor. */
 	CLASS_NotPlaceable        = 0x00000200u,
@@ -601,7 +601,7 @@ enum EObjectFlags : uint32
 	// The group of flags tracks the stages of the lifetime of a uobject
 	RF_NeedInitialization		=0x00000200,	///< This object has not completed its initialization process. Cleared when ~FObjectInitializer completes
 	RF_NeedLoad					=0x00000400,	///< During load, indicates object needs loading.
-	RF_KeepForCooker UE_DEPRECATED(5.6, "This flag is now obsolete as the functionality is unnecessary.")			=0x00000800,	///< Keep this object during garbage collection because it's still being used by the cooker
+	RF_HasPartials				=0x00000800,	///< Set if object has partials (this means memory allocation is offsetted since partials are in front of object ptr)
 	RF_GeneratePublicHash		=0x00000800,	///< During cooking, generate a public hash for this object even is not marked RF_Public.
 	RF_NeedPostLoad				=0x00001000,	///< Object needs to be postloaded.
 	RF_NeedPostLoadSubobjects	=0x00002000,	///< During load, indicates that the object still needs to instance subobjects and fixup serialized component references
@@ -786,6 +786,7 @@ struct FReferencerInformationList
 #define UPROPERTY(...)
 #define UFUNCTION(...)
 #define USTRUCT(...)
+#define UPARTIAL(...)
 #define UMETA(...)
 #define UPARAM(...)
 #define UENUM(...)
diff --git a/Engine/Source/Runtime/CoreUObject/Public/UObject/RegisterCompiledInObjects.h b/Engine/Source/Runtime/CoreUObject/Public/UObject/RegisterCompiledInObjects.h
--- a/Engine/Source/Runtime/CoreUObject/Public/UObject/RegisterCompiledInObjects.h
+++ b/Engine/Source/Runtime/CoreUObject/Public/UObject/RegisterCompiledInObjects.h
@@ -15,6 +15,7 @@ class UClass;
 class UDelegateFunction;
 class UObject;
 class UPackage;
+namespace UE::CoreUObject::Private { struct FPartialClass; }
 
 // Utility struct for passing many arrays with designated initializer for clarity 
 struct FCompiledInObjects
@@ -42,6 +43,8 @@ struct FCompiledInObjects
 	// Other types of objects (e.g. UDelegateFunction) which can be initialized by casting after
 	// classes are propery initialized.
 	TConstArrayView<UObject*> Others;
+	
+	TConstArrayView<UE::CoreUObject::Private::FPartialClass*> Partials;
 
 	inline TConstArrayView<UClass*> GetClasses() const
 	{
diff --git a/Engine/Source/Runtime/CoreUObject/Public/UObject/UObjectAllocator.h b/Engine/Source/Runtime/CoreUObject/Public/UObject/UObjectAllocator.h
--- a/Engine/Source/Runtime/CoreUObject/Public/UObject/UObjectAllocator.h
+++ b/Engine/Source/Runtime/CoreUObject/Public/UObject/UObjectAllocator.h
@@ -41,16 +41,18 @@ public:
 	 * @param Size size of uobject to allocate
 	 * @param Alignment alignment of uobject to allocate
 	 * @param bAllowPermanent if true, allow allocation in the permanent object pool, if it fits
+	 * @param OffsetToObject offset of where the UObjectBase* starts. Should be a positive value to create memory in front of returned pointer
 	 * @return newly allocated UObjectBase (not really a UObjectBase yet, no constructor like thing has been called).
 	 */
-	COREUOBJECT_API UObjectBase* AllocateUObject(int32 Size, int32 Alignment, bool bAllowPermanent);
+	COREUOBJECT_API UObjectBase* AllocateUObject(int32 Size, int32 Alignment, bool bAllowPermanent, int32 OffsetToObject = 0);
 
 	/**
 	 * Returns a UObjectBase to the free store, unless it is in the permanent object pool
 	 *
 	 * @param Object object to free
+	 * @param OffsetToObject offset to Object in allocation. Should match the value sent into AllocateUObject
 	 */
-	COREUOBJECT_API void FreeUObject(UObjectBase *Object) const;
+	COREUOBJECT_API void FreeUObject(UObjectBase *Object, int32 OffsetToObject = 0) const;
 };
 
 /** Global UObjectBase allocator							*/
diff --git a/Engine/Source/Runtime/CoreUObject/Public/UObject/UObjectBase.h b/Engine/Source/Runtime/CoreUObject/Public/UObject/UObjectBase.h
--- a/Engine/Source/Runtime/CoreUObject/Public/UObject/UObjectBase.h
+++ b/Engine/Source/Runtime/CoreUObject/Public/UObject/UObjectBase.h
@@ -32,7 +32,7 @@ class UScriptStruct;
 namespace UE::CoreUObject::Private
 {
 	struct FObjectHandleUtils;
-
+	struct FPartialClass;
 	class FCompiledInObjectLinker;
 }
 
@@ -582,6 +582,25 @@ struct FStructRegisterCompiledInInfo
 	FStructReloadVersionInfo VersionInfo;
 };
 
+
+struct FPartialReloadVersionInfo
+{
+#if WITH_RELOAD
+	SIZE_T Size = 0;
+	uint32 Hash = 0;
+#endif
+};
+
+using FPartialRegistrationInfo = TRegistrationInfo<UE::CoreUObject::Private::FPartialClass, FPartialReloadVersionInfo>;
+
+struct FPartialRegisterCompiledInInfo
+{
+	UE::CoreUObject::Private::FPartialClass* (*OuterRegister)(ETypeConstructPhase);
+	const TCHAR* Name;
+	FPartialRegistrationInfo* Info;
+	FPartialReloadVersionInfo VersionInfo;
+};
+
 /**
  * Adds a struct registration and version information. The InInfo parameter must be static.
  */
@@ -652,7 +671,7 @@ COREUOBJECT_API void RegisterCompiledInInfo(UPackage* (*InOuterRegister)(ETypeCo
 /**
  * Register compiled in information for multiple classes, structures, and enumerations
  */
-COREUOBJECT_API void RegisterCompiledInInfo(const TCHAR* PackageName, const FClassRegisterCompiledInInfo* ClassInfo, size_t NumClassInfo, const FStructRegisterCompiledInInfo* StructInfo, size_t NumStructInfo, const FEnumRegisterCompiledInInfo* EnumInfo, size_t NumEnumInfo);
+COREUOBJECT_API void RegisterCompiledInInfo(const TCHAR *PackageName, const FClassRegisterCompiledInInfo *ClassInfo, size_t NumClassInfo, const FStructRegisterCompiledInInfo *StructInfo, size_t NumStructInfo, const FEnumRegisterCompiledInInfo *EnumInfo, size_t NumEnumInfo, const FPartialRegisterCompiledInInfo *PartialsInfo, size_t NumPartialsInfo);
 
 #endif // !UE_WITH_CONSTINIT_UOBJECT
 
diff --git a/Engine/Source/Runtime/CoreUObject/Public/UObject/UObjectPartials.h b/Engine/Source/Runtime/CoreUObject/Public/UObject/UObjectPartials.h
new file mode 100644
--- /dev/null
+++ b/Engine/Source/Runtime/CoreUObject/Public/UObject/UObjectPartials.h
@@ -0,0 +1,75 @@
+// Copyright Epic Games, Inc. All Rights Reserved.
+
+#pragma once
+
+#include "UObject/ObjectMacros.h"
+#include "UObject/UObjectGlobals.h"
+
+class FLifetimeProperty;
+
+namespace UE::CoreUObject::Private
+{
+
+////////////////////////////////////////////////////////////////////////////////////////////////////
+
+// Function pointer types for partial lifecycle
+using FPartialConstructorFunc = void(void* PartialMemory);
+using FPartialDestructorFunc = void(void* PartialMemory);
+using FPartialBeginDestroyFunc = void(void* PartialMemory);
+using FPartialGetLifetimeReplicatedPropsFunc = void(void* PartialMemory, TArray<FLifetimeProperty>& OutLifetimeProps);
+
+////////////////////////////////////////////////////////////////////////////////////////////////////
+
+// Partial registration node (intrusive linked list). This is on purpose not in namespace in order to keep codegen files small (less text)
+struct FPartialClass
+{
+	const TCHAR* Name = nullptr;
+	int32 Size = 0;
+	int32 Alignment = 0;
+
+#if UE_WITH_CONSTINIT_UOBJECT
+	UClass* Class = nullptr;
+	FProperty* FirstProperty = nullptr;
+	UField* FirstChild = nullptr;
+	UField* EndChild = nullptr; // This is the first field _after_ the fields owned by this partial
+#else
+	UClass* (*Class)(ETypeConstructPhase) = nullptr;
+	const UECodeGen_Private::FPropertyParamsBase* const* Properties = nullptr;
+	int32 NumProperties = 0;
+
+	UFunction* (**FunctionConstructors)(ETypeConstructPhase) = nullptr;  // Array of function construction function pointers (for UFunction creation)
+	const void* NativeFunctions = nullptr;  // Array of FClassNativeFunction (for native binding registration)
+	int32 NumFunctions = 0;
+#endif
+
+	int32* PartialOffsetPtr = nullptr; // Pointer to the native partial struct's static PartialOffset member (for fast access when calling GetPartial<Type>())
+
+	FPartialConstructorFunc* Constructor = nullptr;
+	FPartialBeginDestroyFunc* BeginDestroy = nullptr;
+	FPartialDestructorFunc* Destructor = nullptr;
+	FPartialGetLifetimeReplicatedPropsFunc* GetLifetimeReplicatedProps = nullptr;
+
+	// Runtime-assigned offset from start of partial region (in bytes)
+	// This is the offset from NativePropertiesSize to the start of this partial
+	mutable int32 PartialOffset = MAX_int32;
+
+	// Intrusive linked list - next partial for the same class
+	FPartialClass* NextPartial = nullptr;
+
+	// This call will link in the partial into the class (sorted by alignment/name)
+	void LinkToClass(UClass& PartialClass);
+};
+
+////////////////////////////////////////////////////////////////////////////////////////////////////
+
+// Run the constructors for partial classes associated with this object/class.
+// Does not do anything with properties added by the partials as those are in the list owned by the extended class
+void ConstructPartials(const UClass& Class, UObject* Obj);
+void BeginDestroyPartials(const UClass& Class, UObject* Obj);
+// UNIMPLEMENTED: IsReadyForFinishDestroy/FinishDestroy
+void DestructPartials(const UClass& Class, UObject* Obj);
+void CallPartialsGetLifetimeReplicatedProps(const UClass& Class, UObject* Obj, TArray<FLifetimeProperty>& OutLifetimeProps);
+
+////////////////////////////////////////////////////////////////////////////////////////////////////
+
+} // UE::CoreUObject::Private
\ No newline at end of file
