[UE] Fix for recurring LLT/FoundationTest failures in Horde for platforms that aren't configured to support a UPS-only serialization test case.

#rnx

#Synced-CL 52049047

diff --git a/Engine/Source/Programs/LowLevelTests/FoundationTests/Tests/UPartialTests.cpp b/Engine/Source/Programs/LowLevelTests/FoundationTests/Tests/UPartialTests.cpp
--- a/Engine/Source/Programs/LowLevelTests/FoundationTests/Tests/UPartialTests.cpp
+++ b/Engine/Source/Programs/LowLevelTests/FoundationTests/Tests/UPartialTests.cpp
@@ -708,8 +708,16 @@ TEST_CASE("Partials InheritanceWithoutPartialAttribute", "[CoreUObject][Partials
 	}
 }
 
+// Some build targets aren't configured to support UPS in low-level test runs and will assert at runtime.
+extern bool CanUseUnversionedPropertySerialization(const ITargetPlatform* Target);
+
 TEST_CASE("Partials PartialProperties", "[CoreUObject][Serialization][Partials]")
 {
+	if (!CanUseUnversionedPropertySerialization(nullptr))
+	{
+		return;
+	}
+
 	USerializationTestObject* OriginalObj = NewObject<USerializationTestObject>();
 	OriginalObj->NativeProperty = 123;
 	OriginalObj->NativeString = TEXT("Modified");
@@ -782,6 +790,11 @@ TEST_CASE("Partials PartialProperties", "[CoreUObject][Serialization][Partials]"
 
 TEST_CASE("Partials InheritedPartialProperties", "[CoreUObject][Serialization][Partials]")
 {
+	if (!CanUseUnversionedPropertySerialization(nullptr))
+	{
+		return;
+	}
+
 	SECTION("Derived class serializes both base and derived partial properties")
 	{
 		UDerivedSerializationTestObject* OriginalObj = NewObject<UDerivedSerializationTestObject>();
@@ -806,6 +819,11 @@ TEST_CASE("Partials InheritedPartialProperties", "[CoreUObject][Serialization][P
 
 TEST_CASE("Partials GetDefaultValue_NegativeOffset", "[CoreUObject][Serialization][Partials]")
 {
+	if (!CanUseUnversionedPropertySerialization(nullptr))
+	{
+		return;
+	}
+
 	USerializationTestObject* Obj = NewObject<USerializationTestObject>();
 	FSerializationTestPartial& Partial = Obj->GetPartial<FSerializationTestPartial>();
 	Partial.PartialString = TEXT("ModifiedValue");
